Deleting assets forces the Mode tab's Recently Placed list to be refreshed.
#ttp 334592 - When an asset is deleted from the content browser, it should be removed from the 'recently placed' list as well to avoid confusion
#branch UE4
#proj Editor.PlacementMode
#add Added FPlacementMode::OnAssetRemoved, which forces a refresh of the recently placed list.  Added a call to this method in the OnAssetRemoved event in AssetRegistryModule.
#add In SPlacementModeTools::RefreshRecentlyPlaced, the list of asset paths is iterated, and those which yield a valid FAssetData object are added to the widget container.  In the case of a just deleted widget, the UObject is only marked as pending delete (it has had its RF_Standalone flag cleared), and a GC sweep cannot yet have happened, so it's necessary to explicitly check for this flag set when iterating.

diff --git a/Engine/Source/Editor/PlacementMode/Private/PlacementMode.cpp b/Engine/Source/Editor/PlacementMode/Private/PlacementMode.cpp
--- a/Engine/Source/Editor/PlacementMode/Private/PlacementMode.cpp
+++ b/Engine/Source/Editor/PlacementMode/Private/PlacementMode.cpp
@@ -35,6 +35,9 @@ FPlacementMode::FPlacementMode()
 FPlacementMode::~FPlacementMode()
 {
 	GEditorModeTools().OnEditorModeChanged().RemoveAll( this );
+
+	FAssetRegistryModule& AssetRegistryModule = FModuleManager::GetModuleChecked<FAssetRegistryModule>(TEXT("AssetRegistry"));
+	AssetRegistryModule.Get().OnAssetRemoved().RemoveAll( this );
 }
 
 void FPlacementMode::Initialize()
@@ -54,6 +57,9 @@ void FPlacementMode::Initialize()
 	}
 
 	GEditorModeTools().OnEditorModeChanged().AddSP( this, &FPlacementMode::EditorModeChanged );
+
+	FAssetRegistryModule& AssetRegistryModule = FModuleManager::GetModuleChecked<FAssetRegistryModule>(TEXT("AssetRegistry"));
+	AssetRegistryModule.Get().OnAssetRemoved().AddSP( this, &FPlacementMode::OnAssetRemoved );
 }
 
 bool FPlacementMode::UsesToolkits() const
@@ -685,6 +691,11 @@ void FPlacementMode::AddToRecentlyPlaced( const TArray< UObject* >& PlacedObject
 	}
 }
 
+void FPlacementMode::OnAssetRemoved( const FAssetData& /*InRemovedAssetData*/ )
+{
+	RecentlyPlacedChanged.Broadcast(RecentlyPlaced);
+}
+
 const TArray< TWeakObjectPtr<UObject> >& FPlacementMode::GetCurrentlyPlacingObjects() const
 {
 	return AssetsToPlace;
diff --git a/Engine/Source/Editor/PlacementMode/Private/PlacementMode.h b/Engine/Source/Editor/PlacementMode/Private/PlacementMode.h
--- a/Engine/Source/Editor/PlacementMode/Private/PlacementMode.h
+++ b/Engine/Source/Editor/PlacementMode/Private/PlacementMode.h
@@ -71,6 +71,8 @@ private:
 
 	void EditorModeChanged( FEdMode* Mode, bool IsEntering );
 
+	void OnAssetRemoved( const FAssetData& InRemovedAssetData );
+
 private:
 
 	TWeakObjectPtr< UActorFactory > PlacementFactory;
diff --git a/Engine/Source/Editor/PlacementMode/Private/SPlacementModeTools.cpp b/Engine/Source/Editor/PlacementMode/Private/SPlacementModeTools.cpp
--- a/Engine/Source/Editor/PlacementMode/Private/SPlacementModeTools.cpp
+++ b/Engine/Source/Editor/PlacementMode/Private/SPlacementModeTools.cpp
@@ -680,17 +680,22 @@ void SPlacementModeTools::RefreshRecentlyPlaced()
 	const TArray< FActorPlacementInfo > RecentlyPlaced = IPlacementModeModule::Get().GetRecentlyPlaced();
 	for ( int Index = 0; Index < RecentlyPlaced.Num(); Index++ )
 	{
-		FAssetData AssetData = AssetRegistryModule.Get().GetAssetByObjectPath( *RecentlyPlaced[Index].ObjectPath );
-
-		if ( AssetData.IsValid() )
+		UObject* Asset = FindObject<UObject>(NULL, *RecentlyPlaced[Index].ObjectPath);
+		// If asset is pending delete, it will not be marked as RF_Standalone, in which case we skip it
+		if (Asset != nullptr && Asset->HasAnyFlags(RF_Standalone))
 		{
-			TArray< FActorFactoryAssetProxy::FMenuItem > AssetMenuOptions;
-			UActorFactory* Factory = FindObject<UActorFactory>( nullptr, *RecentlyPlaced[Index].Factory );
+			FAssetData AssetData = AssetRegistryModule.Get().GetAssetByObjectPath(*RecentlyPlaced[Index].ObjectPath);
 
-			RecentlyPlacedContainer->AddSlot()
-			[
-				SNew( SPlacementAssetEntry, Factory, AssetData )
-			];
+			if (AssetData.IsValid())
+			{
+				TArray< FActorFactoryAssetProxy::FMenuItem > AssetMenuOptions;
+				UActorFactory* Factory = FindObject<UActorFactory>(nullptr, *RecentlyPlaced[Index].Factory);
+
+				RecentlyPlacedContainer->AddSlot()
+					[
+						SNew(SPlacementAssetEntry, Factory, AssetData)
+					];
+			}
 		}
 	}
 }
