Renaming an asset which appears in the recently placed list will update the recently placed list accordingly.
#ttp 337370 - EDITOR: CRASH: Placing an actor that has been renamed from the recently placed list results in an ensure
#branch UE4
#proj Editor.PlacementMode
#add Added new delegate to the OnAssetRenamed event (in FPlacementMode::Initialize / FPlacementMode::~FPlacementMode).
#add Created FPlacementMode::OnAssetRenamed which goes through the RecentlyPlaced array and amends the ObjectPath to the new path.

diff --git a/Engine/Source/Editor/PlacementMode/Private/PlacementMode.cpp b/Engine/Source/Editor/PlacementMode/Private/PlacementMode.cpp
--- a/Engine/Source/Editor/PlacementMode/Private/PlacementMode.cpp
+++ b/Engine/Source/Editor/PlacementMode/Private/PlacementMode.cpp
@@ -38,6 +38,7 @@ FPlacementMode::~FPlacementMode()
 
 	FAssetRegistryModule& AssetRegistryModule = FModuleManager::GetModuleChecked<FAssetRegistryModule>(TEXT("AssetRegistry"));
 	AssetRegistryModule.Get().OnAssetRemoved().RemoveAll( this );
+	AssetRegistryModule.Get().OnAssetRenamed().RemoveAll( this );
 }
 
 void FPlacementMode::Initialize()
@@ -60,6 +61,7 @@ void FPlacementMode::Initialize()
 
 	FAssetRegistryModule& AssetRegistryModule = FModuleManager::GetModuleChecked<FAssetRegistryModule>(TEXT("AssetRegistry"));
 	AssetRegistryModule.Get().OnAssetRemoved().AddSP( this, &FPlacementMode::OnAssetRemoved );
+	AssetRegistryModule.Get().OnAssetRenamed().AddSP( this, &FPlacementMode::OnAssetRenamed );
 }
 
 bool FPlacementMode::UsesToolkits() const
@@ -696,6 +698,20 @@ void FPlacementMode::OnAssetRemoved( const FAssetData& /*InRemovedAssetData*/ )
 	RecentlyPlacedChanged.Broadcast(RecentlyPlaced);
 }
 
+void FPlacementMode::OnAssetRenamed( const FAssetData& AssetData, const FString& OldObjectPath )
+{
+	for (auto& RecentlyPlacedItem : RecentlyPlaced)
+	{
+		if (RecentlyPlacedItem.ObjectPath == OldObjectPath)
+		{
+			RecentlyPlacedItem.ObjectPath = AssetData.ObjectPath.ToString();
+			break;
+		}
+	}
+
+	RecentlyPlacedChanged.Broadcast(RecentlyPlaced);
+}
+
 const TArray< TWeakObjectPtr<UObject> >& FPlacementMode::GetCurrentlyPlacingObjects() const
 {
 	return AssetsToPlace;
diff --git a/Engine/Source/Editor/PlacementMode/Private/PlacementMode.h b/Engine/Source/Editor/PlacementMode/Private/PlacementMode.h
--- a/Engine/Source/Editor/PlacementMode/Private/PlacementMode.h
+++ b/Engine/Source/Editor/PlacementMode/Private/PlacementMode.h
@@ -72,6 +72,7 @@ private:
 	void EditorModeChanged( FEdMode* Mode, bool IsEntering );
 
 	void OnAssetRemoved( const FAssetData& InRemovedAssetData );
+	void OnAssetRenamed( const FAssetData& AssetData, const FString& OldObjectPath );
 
 private:
 
