Made it possible for asset editors to maintain their own FEditorModeTools lists

diff --git a/Engine/Source/Editor/LevelEditor/Private/SLevelEditor.cpp b/Engine/Source/Editor/LevelEditor/Private/SLevelEditor.cpp
--- a/Engine/Source/Editor/LevelEditor/Private/SLevelEditor.cpp
+++ b/Engine/Source/Editor/LevelEditor/Private/SLevelEditor.cpp
@@ -160,7 +160,7 @@ void SLevelEditor::Construct( const SLevelEditor::FArguments& InArgs)
 	BindCommands();
 
 	// We need to register when modes list changes so that we can refresh the auto generated commands.
-	GEditorModeTools().OnRegisteredModesChanged().AddRaw(this, &SLevelEditor::RefreshEditorModeCommands);
+	FEditorModeRegistry::Get().OnRegisteredModesChanged().AddRaw(this, &SLevelEditor::RefreshEditorModeCommands);
 
 	// @todo This is a hack to get this working for now. This won't work with multiple worlds
 	GEditor->GetEditorWorldContext(true).AddRef(World);
@@ -261,7 +261,7 @@ SLevelEditor::~SLevelEditor()
 	
 	GetMutableDefault<UEditorExperimentalSettings>()->OnSettingChanged().RemoveAll( this );
 	GEditor->AccessEditorUserSettings().OnUserSettingChanged().RemoveAll( this );
-	GEditorModeTools().OnRegisteredModesChanged().RemoveAll( this );
+	FEditorModeRegistry::Get().OnRegisteredModesChanged().RemoveAll( this );
 
 	FEditorDelegates::MapChange.RemoveAll(this);
 
@@ -1202,17 +1202,17 @@ FName SLevelEditor::GetEditorModeTabId( FEditorModeID ModeID )
 void SLevelEditor::ToggleEditorMode( FEditorModeID ModeID )
 {
 	// *Important* - activate the mode first since FEditorModeTools::DeactivateMode will
-	// activate EM_Default when the stack becomes empty, resulting in multiple active visible modes.
-	GEditorModeTools().ActivateMode( ModeID );
+	// activate the default mode when the stack becomes empty, resulting in multiple active visible modes.
+	GLevelEditorModeTools().ActivateMode( ModeID );
 
 	// Find and disable any other 'visible' modes since we only ever allow one of those active at a time.
 	TArray<FEdMode*> ActiveModes;
-	GEditorModeTools().GetActiveModes( ActiveModes );
+	GLevelEditorModeTools().GetActiveModes( ActiveModes );
 	for ( FEdMode* Mode : ActiveModes )
 	{
-		if ( Mode->IsVisible() && Mode->GetID() != ModeID )
+		if ( Mode->GetID() != ModeID && Mode->GetModeInfo().bVisible )
 		{
-			GEditorModeTools().DeactivateMode( Mode->GetID() );
+			GLevelEditorModeTools().DeactivateMode( Mode->GetID() );
 		}
 	}
 
@@ -1222,15 +1222,15 @@ void SLevelEditor::ToggleEditorMode( FEditorModeID ModeID )
 	TSharedRef<SDockTab> ToolboxTab = LevelEditorTabManager->InvokeTab( FTabId("LevelEditorToolBox") );
 
 	//// If it's already active deactivate the mode
-	//if ( GEditorModeTools().IsModeActive( ModeID ) )
+	//if ( GLevelEditorModeTools().IsModeActive( ModeID ) )
 	//{
-	//	//GEditorModeTools().DeactivateAllModes();
-	//	//GEditorModeTools().DeactivateMode( ModeID );
+	//	//GLevelEditorModeTools().DeactivateAllModes();
+	//	//GLevelEditorModeTools().DeactivateMode( ModeID );
 	//}
 	//else // Activate the mode and create the tab for it.
 	//{
-	//	GEditorModeTools().DeactivateAllModes();
-	//	GEditorModeTools().ActivateMode( ModeID );
+	//	GLevelEditorModeTools().DeactivateAllModes();
+	//	GLevelEditorModeTools().ActivateMode( ModeID );
 
 	//	//FLevelEditorModule& LevelEditorModule = FModuleManager::GetModuleChecked<FLevelEditorModule>( "LevelEditor" );
 	//	//TSharedPtr<FTabManager> LevelEditorTabManager = LevelEditorModule.GetLevelEditorTabManager();
@@ -1239,33 +1239,23 @@ void SLevelEditor::ToggleEditorMode( FEditorModeID ModeID )
 	//}
 }
 
-static bool IsDefaultModeActive()
+bool SLevelEditor::IsModeActive( FEditorModeID ModeID )
 {
-	TArray<FEdMode*> ActiveModes;
-	GEditorModeTools().GetActiveModes( ActiveModes );
-	for ( FEdMode* Mode : ActiveModes )
+	// The level editor changes the default mode to placement
+	if ( ModeID == FBuiltinEditorModes::EM_Placement )
 	{
-		if ( Mode->IsVisible() )
+		// Only return true if this is the *only* active mode
+		TArray<FEdMode*> ActiveModes;
+		GLevelEditorModeTools().GetActiveModes(ActiveModes);
+		for( FEdMode* Mode : ActiveModes )
 		{
-			if ( !( Mode->GetID() == FBuiltinEditorModes::EM_Placement ||
-					Mode->GetID() == FBuiltinEditorModes::EM_Default ) )
+			if( Mode->GetModeInfo().bVisible && Mode->GetID() != FBuiltinEditorModes::EM_Placement )
 			{
 				return false;
 			}
 		}
 	}
-
-	return GEditorModeTools().IsModeActive( FBuiltinEditorModes::EM_Placement );
-}
-
-bool SLevelEditor::IsModeActive( FEditorModeID ModeID )
-{
-	if ( ModeID == FBuiltinEditorModes::EM_Placement || ModeID == FBuiltinEditorModes::EM_Default )
-	{
-		return IsDefaultModeActive();
-	}
-
-	return GEditorModeTools().IsModeActive( ModeID );
+	return GLevelEditorModeTools().IsModeActive( ModeID );
 }
 
 void SLevelEditor::RefreshEditorModeCommands()
@@ -1281,32 +1271,17 @@ void SLevelEditor::RefreshEditorModeCommands()
 	// We need to remap all the actions to commands.
 	const FLevelEditorModesCommands& Commands = FLevelEditorModesCommands::Get();
 
-	TArray<FEdMode*> Modes;
-	GEditorModeTools().GetModes(Modes);
-
-	int editorMode = 0;
-
-	struct FCompareEdModeByPriority
-	{
-		FORCEINLINE bool operator()(const FEdMode& A, const FEdMode& B) const
-		{
-			return A.GetPriorityOrder() < B.GetPriorityOrder();
-		}
-	};
-
-	Modes.Sort(FCompareEdModeByPriority());
-
 	int commandIndex = 0;
-	for ( FEdMode* Mode : Modes )
+	for( const FEditorModeInfo& Mode : FEditorModeRegistry::Get().GetSortedModeInfo() )
 	{
 		// If the mode isn't visible don't create a menu option for it.
-		if ( !Mode->IsVisible() )
+		if( !Mode.bVisible )
 		{
 			continue;
 		}
 
-		FName EditorModeTabName = GetEditorModeTabId( Mode->GetID() );
-		FName EditorModeCommandName = FName(*(FString("EditorMode.") + Mode->GetID().ToString()));
+		FName EditorModeTabName = GetEditorModeTabId( Mode.ID );
+		FName EditorModeCommandName = FName(*(FString("EditorMode.") + Mode.ID.ToString()));
 
 		TSharedPtr<FUICommandInfo> EditorModeCommand = 
 			FInputBindingManager::Get().FindCommandInContext(Commands.GetContextName(), EditorModeCommandName);
@@ -1316,9 +1291,9 @@ void SLevelEditor::RefreshEditorModeCommands()
 		{
 			LevelEditorCommands->MapAction(
 				Commands.EditorModeCommands[commandIndex],
-				FExecuteAction::CreateStatic( &SLevelEditor::ToggleEditorMode, Mode->GetID() ),
+				FExecuteAction::CreateStatic( &SLevelEditor::ToggleEditorMode, Mode.ID ),
 				FCanExecuteAction(),
-				FIsActionChecked::CreateStatic( &SLevelEditor::IsModeActive, Mode->GetID() ));
+				FIsActionChecked::CreateStatic( &SLevelEditor::IsModeActive, Mode.ID ));
 		}
 
 		// Register the new tab spawner, and unregister the old one if it exists
@@ -1334,34 +1309,43 @@ void SLevelEditor::RefreshEditorModeCommands()
 
 		commandIndex++;
 	}
-}
 
-TSharedRef<SDockTab> SLevelEditor::SpawnLevelEditorModeTab( const FSpawnTabArgs& Args, FEdMode* EditorMode )
-{
-	if (!GEditorModeTools().IsModeActive(EditorMode->GetID()))
+	for( const auto& ToolBoxTab : ToolBoxTabs )
 	{
-		GEditorModeTools().ActivateMode(EditorMode->GetID());
+		auto Tab = ToolBoxTab.Pin();
+		if( Tab.IsValid() )
+		{
+			Tab->OnEditorModeCommandsChanged();
+		}
 	}
+}
 
-	TSharedPtr<SDockTab> NewEditorModeTab;
-	TSharedPtr<SLevelEditorModeContent> NewToolBox;
+//TSharedRef<SDockTab> SLevelEditor::SpawnLevelEditorModeTab( const FSpawnTabArgs& Args, FEdMode* EditorMode )
+//{
+	//if (!GLevelEditorModeTools().IsModeActive(EditorMode->GetID()))
+	//{
+	//	GLevelEditorModeTools().ActivateMode(EditorMode->GetID());
+	//}
 
-	SAssignNew( NewEditorModeTab, SDockTab )
-		.Icon( EditorMode->GetIcon().GetSmallIcon() )
-		.Label( EditorMode->GetName() );
+	//TSharedPtr<SDockTab> NewEditorModeTab;
+	//TSharedPtr<SLevelEditorModeContent> NewToolBox;
 
-	NewEditorModeTab->SetContent(
-		SNew( STutorialWrapper, TEXT("ToolsPanel") )
-		[
-			SAssignNew( NewToolBox, SLevelEditorModeContent, SharedThis( this ), NewEditorModeTab.ToSharedRef(), EditorMode )
-			.IsEnabled( FSlateApplication::Get().GetNormalExecutionAttribute() )
-		]
-	);
+	//SAssignNew( NewEditorModeTab, SDockTab )
+	//	.Icon( EditorMode->GetIcon().GetSmallIcon() )
+	//	.Label( EditorMode->GetName() );
 
-	ModesTabs.Add( NewToolBox );
+	//NewEditorModeTab->SetContent(
+	//	SNew( STutorialWrapper, TEXT("ToolsPanel") )
+	//	[
+	//		SAssignNew( NewToolBox, SLevelEditorModeContent, SharedThis( this ), NewEditorModeTab.ToSharedRef(), EditorMode )
+	//		.IsEnabled( FSlateApplication::Get().GetNormalExecutionAttribute() )
+	//	]
+	//);
 
-	return NewEditorModeTab.ToSharedRef();
-}
+	//ModesTabs.Add( NewToolBox );
+
+	//return NewEditorModeTab.ToSharedRef();
+//}
 
 FReply SLevelEditor::OnKeyDown( const FGeometry& MyGeometry, const FKeyboardEvent& InKeyboardEvent )
 {
diff --git a/Engine/Source/Editor/LevelEditor/Private/SLevelEditor.h b/Engine/Source/Editor/LevelEditor/Private/SLevelEditor.h
--- a/Engine/Source/Editor/LevelEditor/Private/SLevelEditor.h
+++ b/Engine/Source/Editor/LevelEditor/Private/SLevelEditor.h
@@ -21,7 +21,7 @@ class SLevelEditor
 public:
 	// @todo why are these public?
 	TSharedRef<SDockTab> SpawnLevelEditorTab( const FSpawnTabArgs& Args, FName TabIdentifier, FString InitializationPayload );
-	TSharedRef<SDockTab> SpawnLevelEditorModeTab( const FSpawnTabArgs& Args, FEdMode* EditorMode );
+	//TSharedRef<SDockTab> SpawnLevelEditorModeTab( const FSpawnTabArgs& Args, FEdMode* EditorMode );
 
 	SLATE_BEGIN_ARGS( SLevelEditor ){}
 
diff --git a/Engine/Source/Editor/LevelEditor/Private/SLevelEditorModeContent.cpp b/Engine/Source/Editor/LevelEditor/Private/SLevelEditorModeContent.cpp
--- a/Engine/Source/Editor/LevelEditor/Private/SLevelEditorModeContent.cpp
+++ b/Engine/Source/Editor/LevelEditor/Private/SLevelEditorModeContent.cpp
@@ -13,7 +13,7 @@
 
 SLevelEditorModeContent::~SLevelEditorModeContent()
 {
-	GEditorModeTools().OnEditorModeChanged().RemoveAll( this );
+	GLevelEditorModeTools().OnEditorModeChanged().RemoveAll( this );
 	GEditor->AccessEditorUserSettings().OnUserSettingChanged().RemoveAll( this );
 }
 
@@ -24,7 +24,7 @@ void SLevelEditorModeContent::Construct( const FArguments& InArgs, const TShared
 	EditorMode = InEditorMode;
 
 	InOwningDocTab->SetOnTabClosed( SDockTab::FOnTabClosedCallback::CreateSP(this, &SLevelEditorModeContent::HandleParentClosed ) );
-	GEditorModeTools().OnEditorModeChanged().AddSP( this, &SLevelEditorModeContent::HandleEditorModeChanged );
+	GLevelEditorModeTools().OnEditorModeChanged().AddSP( this, &SLevelEditorModeContent::HandleEditorModeChanged );
 	GEditor->AccessEditorUserSettings().OnUserSettingChanged().AddSP( this, &SLevelEditorModeContent::HandleUserSettingsChange );
 
 	ChildSlot
@@ -124,9 +124,9 @@ void SLevelEditorModeContent::OnToolkitHostingFinished( const TSharedRef< class
 
 void SLevelEditorModeContent::HandleParentClosed( TSharedRef<SDockTab> TabBeingClosed )
 {
-	if ( GEditorModeTools().IsModeActive(EditorMode->GetID()) )
+	if ( GLevelEditorModeTools().IsModeActive(EditorMode->GetID()) )
 	{
-		GEditorModeTools().DeactivateMode(EditorMode->GetID());
+		GLevelEditorModeTools().DeactivateMode(EditorMode->GetID());
 	}
 }
 
diff --git a/Engine/Source/Editor/LevelEditor/Private/SLevelEditorToolBox.cpp b/Engine/Source/Editor/LevelEditor/Private/SLevelEditorToolBox.cpp
--- a/Engine/Source/Editor/LevelEditor/Private/SLevelEditorToolBox.cpp
+++ b/Engine/Source/Editor/LevelEditor/Private/SLevelEditorToolBox.cpp
@@ -73,6 +73,11 @@ void SLevelEditorToolBox::HandleUserSettingsChange( FName PropertyName )
 	UpdateModeToolBar();
 }
 
+void SLevelEditorToolBox::OnEditorModeCommandsChanged()
+{
+	UpdateModeToolBar();
+}
+
 void SLevelEditorToolBox::UpdateModeToolBar()
 {
 	FLevelEditorModule& LevelEditorModule = FModuleManager::GetModuleChecked<FLevelEditorModule>( "LevelEditor");
@@ -86,28 +91,15 @@ void SLevelEditorToolBox::UpdateModeToolBar()
 
 		const FLevelEditorModesCommands& Commands = LevelEditorModule.GetLevelEditorModesCommands();
 
-		TArray<FEdMode*> Modes;
-		GEditorModeTools().GetModes( Modes );
-
-		struct FCompareEdModeByPriority
-		{
-			FORCEINLINE bool operator()( const FEdMode& A, const FEdMode& B ) const
-			{
-				return A.GetPriorityOrder() < B.GetPriorityOrder();
-			}
-		};
-
-		Modes.Sort( FCompareEdModeByPriority() );
-
-		for ( FEdMode* Mode : Modes )
+		for ( const FEditorModeInfo& Mode : FEditorModeRegistry::Get().GetSortedModeInfo() )
 		{
 			// If the mode isn't visible don't create a menu option for it.
-			if ( !Mode->IsVisible() )
+			if ( !Mode.bVisible )
 			{
 				continue;
 			}
 
-			FName EditorModeCommandName = FName( *( FString( "EditorMode." ) + Mode->GetID().ToString() ) );
+			FName EditorModeCommandName = FName( *( FString( "EditorMode." ) + Mode.ID.ToString() ) );
 
 			TSharedPtr<FUICommandInfo> EditorModeCommand =
 				FInputBindingManager::Get().FindCommandInContext( Commands.GetContextName(), EditorModeCommandName );
@@ -121,8 +113,7 @@ void SLevelEditorToolBox::UpdateModeToolBar()
 			const FUIAction* UIAction = EditorModeTools.GetTopCommandList()->GetActionForCommand( EditorModeCommand );
 			if ( ensure( UIAction ) )
 			{
-				EditorModeTools.AddToolBarButton( EditorModeCommand, Mode->GetID(), Mode->GetName(), Mode->GetName(), Mode->GetIcon(), Mode->GetID() );// , EUserInterfaceActionType::ToggleButton );
-				//EditorModeTools.AddToolBarButton( *UIAction, Mode->GetID(), Mode->GetName(), Mode->GetName(), Mode->GetIcon(), EUserInterfaceActionType::ToggleButton );
+				EditorModeTools.AddToolBarButton( EditorModeCommand, Mode.ID, Mode.Name, Mode.Name, Mode.IconBrush, Mode.ID );// , EUserInterfaceActionType::ToggleButton );
 			}
 		}
 	}
diff --git a/Engine/Source/Editor/LevelEditor/Private/SLevelEditorToolBox.h b/Engine/Source/Editor/LevelEditor/Private/SLevelEditorToolBox.h
--- a/Engine/Source/Editor/LevelEditor/Private/SLevelEditorToolBox.h
+++ b/Engine/Source/Editor/LevelEditor/Private/SLevelEditorToolBox.h
@@ -20,7 +20,10 @@ public:
 	void OnToolkitHostingStarted( const TSharedRef< class IToolkit >& Toolkit );
 
 	/** Called by SLevelEditor to notify the toolbox about an existing toolkit no longer being hosted */
-	void OnToolkitHostingFinished( const TSharedRef< class IToolkit >& Toolkit );
+	void OnToolkitHostingFinished(const TSharedRef< class IToolkit >& Toolkit);
+
+	/** Handles updating the mode toolbar when the registered mode commands change */
+	void OnEditorModeCommandsChanged();
 
 private:
 
diff --git a/Engine/Source/Editor/LevelEditor/Private/SLevelViewport.cpp b/Engine/Source/Editor/LevelEditor/Private/SLevelViewport.cpp
--- a/Engine/Source/Editor/LevelEditor/Private/SLevelViewport.cpp
+++ b/Engine/Source/Editor/LevelEditor/Private/SLevelViewport.cpp
@@ -1887,23 +1887,23 @@ FLevelEditorViewportInstanceSettings SLevelViewport::LoadLegacyConfigFromIni(con
 
 void SLevelViewport::OnSetBookmark( int32 BookmarkIndex )
 {
-	GEditorModeTools().SetBookmark( BookmarkIndex, LevelViewportClient.Get() );
+	GLevelEditorModeTools().SetBookmark( BookmarkIndex, LevelViewportClient.Get() );
 }
 
 void SLevelViewport::OnJumpToBookmark( int32 BookmarkIndex )
 {
 	const bool bShouldRestoreLevelVisibility = true;
-	GEditorModeTools().JumpToBookmark( BookmarkIndex, bShouldRestoreLevelVisibility, LevelViewportClient.Get() );
+	GLevelEditorModeTools().JumpToBookmark( BookmarkIndex, bShouldRestoreLevelVisibility, LevelViewportClient.Get() );
 }
 
 void SLevelViewport::OnClearBookMark( int32 BookmarkIndex )
 {
-	GEditorModeTools().ClearBookmark( BookmarkIndex, LevelViewportClient.Get() );
+	GLevelEditorModeTools().ClearBookmark( BookmarkIndex, LevelViewportClient.Get() );
 }
 
 void SLevelViewport::OnClearAllBookMarks()
 {
-	GEditorModeTools().ClearAllBookmarks( LevelViewportClient.Get() );
+	GLevelEditorModeTools().ClearAllBookmarks( LevelViewportClient.Get() );
 }
 
 void SLevelViewport::OnToggleAllowMatineePreview()
diff --git a/Engine/Source/Editor/LevelEditor/Private/SLevelViewportToolBar.cpp b/Engine/Source/Editor/LevelEditor/Private/SLevelViewportToolBar.cpp
--- a/Engine/Source/Editor/LevelEditor/Private/SLevelViewportToolBar.cpp
+++ b/Engine/Source/Editor/LevelEditor/Private/SLevelViewportToolBar.cpp
@@ -366,7 +366,7 @@ static void OnGenerateSetBookmarkMenu( FMenuBuilder& MenuBuilder )
 static void OnGenerateClearBookmarkMenu( FMenuBuilder& MenuBuilder , TWeakPtr<class SLevelViewport> Viewport )
 {
 	// Add a menu entry for each bookmark
-	FEditorModeTools& Tools = GEditorModeTools();
+	FEditorModeTools& Tools = GLevelEditorModeTools();
 
 	// Get the viewport client to pass down to the CheckBookmark function
 	FLevelEditorViewportClient& ViewportClient = Viewport.Pin()->GetLevelViewportClient();	
@@ -520,7 +520,7 @@ static void OnGenerateActorLockingMenu( FMenuBuilder& MenuBuilder, TWeakPtr<SLev
  */
 static void OnGenerateBookmarkMenu( FMenuBuilder& MenuBuilder , TWeakPtr<class SLevelViewport> Viewport )
 {
-	FEditorModeTools& Tools = GEditorModeTools();
+	FEditorModeTools& Tools = GLevelEditorModeTools();
 
 	// true if a bookmark was found. 
 	bool bFoundBookmark = false;
diff --git a/Engine/Source/Editor/LevelEditor/Private/SToolkitDisplay.cpp b/Engine/Source/Editor/LevelEditor/Private/SToolkitDisplay.cpp
--- a/Engine/Source/Editor/LevelEditor/Private/SToolkitDisplay.cpp
+++ b/Engine/Source/Editor/LevelEditor/Private/SToolkitDisplay.cpp
@@ -142,7 +142,7 @@ FText SLevelEditorActiveToolkit::GetToolkitTextLabel() const
 {
 	if( ActiveToolkitType == ELevelEditorActiveToolkit::LegacyEditorMode )
 	{
-		return EditorMode->GetName();
+		return EditorMode->GetModeInfo().Name;
 	}
 	else if( ActiveToolkitType == ELevelEditorActiveToolkit::Toolkit )
 	{
@@ -172,7 +172,7 @@ void SLevelEditorActiveToolkit::DeactivateMode()
 	{
 		// Exit the mode.  This will immediately call back to remove the active toolkit from our list
 		check( EditorMode != NULL );
-		GEditorModeTools().DeactivateMode( EditorMode->GetID() );
+		GLevelEditorModeTools().DeactivateMode( EditorMode->GetID() );
 	}
 	else if( ActiveToolkitType == ELevelEditorActiveToolkit::Toolkit )
 	{
@@ -238,12 +238,12 @@ void SToolkitDisplay::Construct( const FArguments& InArgs, const TSharedRef< cla
 		];
 
 	// Register with the mode system to find out when a mode is entered or exited
-	GEditorModeTools().OnEditorModeChanged().AddSP( SharedThis( this ), &SToolkitDisplay::OnEditorModeChanged );
+	GLevelEditorModeTools().OnEditorModeChanged().AddSP( SharedThis( this ), &SToolkitDisplay::OnEditorModeChanged );
 
 	// Find all of the current active modes and toolkits and add those right away.  This widget could have been created "late"!
 	{
 		TArray< FEdMode* > ActiveModes;
-		GEditorModeTools().GetActiveModes( ActiveModes );
+		GLevelEditorModeTools().GetActiveModes( ActiveModes );
 		for( auto EdModeIt = ActiveModes.CreateConstIterator(); EdModeIt; ++EdModeIt )
 		{
 			// We don't care about the default editor mode.  Just ignore it.
@@ -271,7 +271,7 @@ SToolkitDisplay::~SToolkitDisplay()
 	}
 
 	// Unregister delegates
-	GEditorModeTools().OnEditorModeChanged().RemoveAll( this );
+	GLevelEditorModeTools().OnEditorModeChanged().RemoveAll( this );
 }
 
 
diff --git a/Engine/Source/Editor/LevelEditor/Public/SLevelViewport.h b/Engine/Source/Editor/LevelEditor/Public/SLevelViewport.h
--- a/Engine/Source/Editor/LevelEditor/Public/SLevelViewport.h
+++ b/Engine/Source/Editor/LevelEditor/Public/SLevelViewport.h
@@ -17,6 +17,7 @@ public:
 		, _Realtime( false )
 		{}
 
+		SLATE_ARGUMENT( TWeakPtr<class FEditorModeTools>, EditorModeTools )
 		SLATE_ARGUMENT( TSharedPtr<class FLevelViewportLayout>, ParentLayout )
 		SLATE_ARGUMENT( TWeakPtr<class SLevelEditor>, ParentLevelEditor )
 		SLATE_ARGUMENT( ELevelViewportType, ViewportType )
diff --git a/Engine/Source/Editor/Levels/Private/LevelCollectionViewModel.cpp b/Engine/Source/Editor/Levels/Private/LevelCollectionViewModel.cpp
--- a/Engine/Source/Editor/Levels/Private/LevelCollectionViewModel.cpp
+++ b/Engine/Source/Editor/Levels/Private/LevelCollectionViewModel.cpp
@@ -1440,9 +1440,9 @@ bool FLevelCollectionViewModel::AddExistingLevel()
 		}
 
 		// For safety
-		if( GEditorModeTools().IsModeActive( FBuiltinEditorModes::EM_Landscape ) )
+		if( GLevelEditorModeTools().IsModeActive( FBuiltinEditorModes::EM_Landscape ) )
 		{
-			GEditorModeTools().ActivateMode( FBuiltinEditorModes::EM_Default );
+			GLevelEditorModeTools().ActivateDefaultMode();
 		}
 
 		// refresh editor windows
@@ -1519,18 +1519,18 @@ void FLevelCollectionViewModel::RemoveSelectedLevels_Executed()
 	}
 
 	// If matinee is opened, and if it belongs to the level being removed, close it
-	if( GEditorModeTools().IsModeActive( FBuiltinEditorModes::EM_InterpEdit ) )
+	if( GLevelEditorModeTools().IsModeActive( FBuiltinEditorModes::EM_InterpEdit ) )
 	{
-		const FEdModeInterpEdit* InterpEditMode = (const FEdModeInterpEdit*)GEditorModeTools().GetActiveMode( FBuiltinEditorModes::EM_InterpEdit );
+		const FEdModeInterpEdit* InterpEditMode = (const FEdModeInterpEdit*)GLevelEditorModeTools().GetActiveMode( FBuiltinEditorModes::EM_InterpEdit );
 
 		if ( InterpEditMode && InterpEditMode->MatineeActor && LevelsToRemove.Contains( InterpEditMode->MatineeActor->GetLevel() ) )
 		{
-			GEditorModeTools().ActivateMode(FBuiltinEditorModes::EM_Default);
+			GLevelEditorModeTools().ActivateDefaultMode();
 		}
 	}
-	else if( GEditorModeTools().IsModeActive( FBuiltinEditorModes::EM_Landscape ) )
+	else if( GLevelEditorModeTools().IsModeActive( FBuiltinEditorModes::EM_Landscape ) )
 	{
-		GEditorModeTools().ActivateMode(FBuiltinEditorModes::EM_Default);
+		GLevelEditorModeTools().ActivateDefaultMode();
 	}
 
 	// Disassociate selected levels from streaming volumes since the levels will be removed
@@ -2095,9 +2095,9 @@ void FLevelCollectionViewModel::AddExistingLevelFromAssetPicker(const TArray<FAs
 			}
 
 			// For safety
-			if (GEditorModeTools().IsModeActive(FBuiltinEditorModes::EM_Landscape))
+			if (GLevelEditorModeTools().IsModeActive(FBuiltinEditorModes::EM_Landscape))
 			{
-				GEditorModeTools().ActivateMode(FBuiltinEditorModes::EM_Default);
+				GLevelEditorModeTools().ActivateDefaultMode();
 			}
 
 			// refresh editor windows
diff --git a/Engine/Source/Editor/Levels/Private/LevelEdMode.cpp b/Engine/Source/Editor/Levels/Private/LevelEdMode.cpp
--- a/Engine/Source/Editor/Levels/Private/LevelEdMode.cpp
+++ b/Engine/Source/Editor/Levels/Private/LevelEdMode.cpp
@@ -12,9 +12,6 @@ FEdModeLevel::FEdModeLevel()
 	: FEdMode()
 	, SelectedLevel( NULL )
 {
-	ID = FBuiltinEditorModes::EM_Level;
-	Name = LOCTEXT("LevelMode", "Level Transform Editing");
-
 	UMaterial* GizmoMaterial = (UMaterial*)StaticLoadObject( UMaterial::StaticClass(),NULL,TEXT("/Engine/EditorMaterials/LevelTransformMaterial.LevelTransformMaterial") );
 	BoxMaterial = UMaterialInstanceDynamic::Create( GizmoMaterial, NULL );
 
@@ -109,7 +106,7 @@ void FEdModeLevel::SetLevel( ULevelStreaming* LevelStream )
 	GEditor->RedrawAllViewports();
 }
 
-bool FEdModeLevel::InputDelta( FLevelEditorViewportClient* InViewportClient, FViewport* InViewport, FVector& InDrag, FRotator& InRot, FVector& InScale )
+bool FEdModeLevel::InputDelta( FEditorViewportClient* InViewportClient, FViewport* InViewport, FVector& InDrag, FRotator& InRot, FVector& InScale )
 {
 	// Only Update the LevelTransform if the user has clicked on the Widget
 	if (InViewportClient->GetCurrentWidgetAxis() != EAxisList::None)
@@ -148,13 +145,13 @@ void FEdModeLevel::Render( const FSceneView* View, FViewport* Viewport, FPrimiti
 	}
 }
 
-bool FEdModeLevel::StartTracking( FLevelEditorViewportClient* InViewportClient, FViewport* InViewport )
+bool FEdModeLevel::StartTracking( FEditorViewportClient* InViewportClient, FViewport* InViewport )
 {	
 	bIsTracking = true;
 	return true;
 }
 
-bool FEdModeLevel::EndTracking( FLevelEditorViewportClient* InViewportClient, FViewport* InViewport )
+bool FEdModeLevel::EndTracking( FEditorViewportClient* InViewportClient, FViewport* InViewport )
 {	
 	bIsTracking = false;
 	if( SelectedLevel && !LevelTransform.Equals( SelectedLevel->LevelTransform ) )
diff --git a/Engine/Source/Editor/Levels/Private/LevelsModule.cpp b/Engine/Source/Editor/Levels/Private/LevelsModule.cpp
--- a/Engine/Source/Editor/Levels/Private/LevelsModule.cpp
+++ b/Engine/Source/Editor/Levels/Private/LevelsModule.cpp
@@ -10,9 +10,7 @@ void FLevelsModule::StartupModule()
 	FLevelsViewCommands::Register();
 
 	// register the editor mode
-	TSharedRef<FEdModeLevel> NewEditorMode = MakeShareable(new FEdModeLevel);
-	GEditorModeTools().RegisterMode(NewEditorMode);
-	EdModeLevel = NewEditorMode;
+	FEditorModeRegistry::Get().RegisterMode<FEdModeLevel>(FBuiltinEditorModes::EM_Level);
 }
 
 
@@ -21,8 +19,7 @@ void FLevelsModule::ShutdownModule()
 	FLevelsViewCommands::Unregister();
 
 	// unregister the editor mode
-	GEditorModeTools().UnregisterMode(EdModeLevel.ToSharedRef());
-	EdModeLevel = NULL;
+	FEditorModeRegistry::Get().UnregisterMode(FBuiltinEditorModes::EM_Level);
 }
 
 
@@ -34,4 +31,4 @@ TSharedRef< SWidget > FLevelsModule::CreateLevelBrowser()
 
 ULevelBrowserSettings::ULevelBrowserSettings(const class FPostConstructInitializeProperties& PCIP)
 	: Super(PCIP)
-{ }
\ No newline at end of file
+{ }
diff --git a/Engine/Source/Editor/Levels/Private/SLevelsView.h b/Engine/Source/Editor/Levels/Private/SLevelsView.h
--- a/Engine/Source/Editor/Levels/Private/SLevelsView.h
+++ b/Engine/Source/Editor/Levels/Private/SLevelsView.h
@@ -37,7 +37,7 @@ public:
 		ViewModel->OnDisplayActorCountChanged().RemoveAll( this );
 		ViewModel->OnDisplayLightmassSizeChanged().RemoveAll( this );
 		ViewModel->OnDisplayFileSizeChanged().RemoveAll( this ); 
-		GEditorModeTools().DeactivateMode( FBuiltinEditorModes::EM_Level );
+		GLevelEditorModeTools().DeactivateMode( FBuiltinEditorModes::EM_Level );
 	}
 
 	/**
diff --git a/Engine/Source/Editor/Levels/Private/SLevelsViewRow.h b/Engine/Source/Editor/Levels/Private/SLevelsViewRow.h
--- a/Engine/Source/Editor/Levels/Private/SLevelsViewRow.h
+++ b/Engine/Source/Editor/Levels/Private/SLevelsViewRow.h
@@ -582,13 +582,13 @@ private:
 	FReply OnToggleLock()
 	{
 		// If were locking a level, lets make sure to close the level transform mode if its the same level currently selected for edit
-		FEdModeLevel* LevelMode = static_cast<FEdModeLevel*>(GEditorModeTools().GetActiveMode( FBuiltinEditorModes::EM_Level ));		
+		FEdModeLevel* LevelMode = static_cast<FEdModeLevel*>(GLevelEditorModeTools().GetActiveMode( FBuiltinEditorModes::EM_Level ));		
 		if( LevelMode )
 		{
 			ULevelStreaming* LevelStreaming = ViewModel->GetLevelStreaming().Get();
 			if( LevelMode->IsEditing( LevelStreaming ) )
 			{
-				GEditorModeTools().DeactivateMode( FBuiltinEditorModes::EM_Level );
+				GLevelEditorModeTools().DeactivateMode( FBuiltinEditorModes::EM_Level );
 			}
 		}
 
@@ -945,19 +945,19 @@ private:
 			return FReply::Handled();
 		}		
 
-		if( !GEditorModeTools().IsModeActive(FBuiltinEditorModes::EM_Level))
+		if( !GLevelEditorModeTools().IsModeActive(FBuiltinEditorModes::EM_Level))
 		{
 			// Activate Level Mode if it was not active
-			GEditorModeTools().ActivateMode( FBuiltinEditorModes::EM_Level );
+			GLevelEditorModeTools().ActivateMode( FBuiltinEditorModes::EM_Level );
 		}
 			
-		FEdModeLevel* ActiveMode = static_cast<FEdModeLevel*>(GEditorModeTools().GetActiveMode( FBuiltinEditorModes::EM_Level ));		
+		FEdModeLevel* ActiveMode = static_cast<FEdModeLevel*>(GLevelEditorModeTools().GetActiveMode( FBuiltinEditorModes::EM_Level ));		
 		check(ActiveMode != NULL);
 
 		if( ActiveMode->IsEditing(LevelStreaming) == true )
 		{
 			// Toggle this mode off if already editing this level
-			GEditorModeTools().DeactivateMode(FBuiltinEditorModes::EM_Level);
+			GLevelEditorModeTools().DeactivateMode(FBuiltinEditorModes::EM_Level);
 
 			// Cache Transform, as it might of changed by doing a view port edit
 			LevelTransform = LevelStreaming->LevelTransform;
@@ -987,7 +987,7 @@ private:
 	bool LevelEditTextTransformAllowed() const
 	{
 		ULevelStreaming* LevelStreaming = ViewModel->GetLevelStreaming().Get();
-		FEdModeLevel* ActiveMode = static_cast<FEdModeLevel*>(GEditorModeTools().GetActiveMode( FBuiltinEditorModes::EM_Level ));		
+		FEdModeLevel* ActiveMode = static_cast<FEdModeLevel*>(GLevelEditorModeTools().GetActiveMode( FBuiltinEditorModes::EM_Level ));		
 		if( ActiveMode && ActiveMode->IsEditing(LevelStreaming) == true )
 		{
 			return false;
diff --git a/Engine/Source/Editor/Levels/Public/LevelEdMode.h b/Engine/Source/Editor/Levels/Public/LevelEdMode.h
--- a/Engine/Source/Editor/Levels/Public/LevelEdMode.h
+++ b/Engine/Source/Editor/Levels/Public/LevelEdMode.h
@@ -30,13 +30,13 @@ public:
 
 	virtual bool AllowWidgetMove() override { return true; }
 
-	virtual bool InputDelta( FLevelEditorViewportClient* InViewportClient, FViewport* InViewport, FVector& InDrag, FRotator& InRot, FVector& InScale ) override;
+	virtual bool InputDelta(FEditorViewportClient* InViewportClient, FViewport* InViewport, FVector& InDrag, FRotator& InRot, FVector& InScale) override;
 
 	virtual void Render( const FSceneView* View, FViewport* Viewport, FPrimitiveDrawInterface* PDI ) override;
 
-	virtual bool StartTracking(FLevelEditorViewportClient* InViewportClient, FViewport* InViewport) override;
+	virtual bool StartTracking(FEditorViewportClient* InViewportClient, FViewport* InViewport) override;
 
-	virtual bool EndTracking(FLevelEditorViewportClient* InViewportClient, FViewport* InViewport) override;
+	virtual bool EndTracking(FEditorViewportClient* InViewportClient, FViewport* InViewport) override;
 
 	virtual bool IsSnapRotationEnabled() override;
 
diff --git a/Engine/Source/Editor/Levels/Public/LevelsModule.h b/Engine/Source/Editor/Levels/Public/LevelsModule.h
--- a/Engine/Source/Editor/Levels/Public/LevelsModule.h
+++ b/Engine/Source/Editor/Levels/Public/LevelsModule.h
@@ -36,8 +36,6 @@ private:
 
 	/** All extender delegates for the levels menus */
 	TArray<FLevelsMenuExtender> LevelsMenuExtenders;
-
-	TSharedPtr<FEdModeLevel> EdModeLevel;
 };
 
 
diff --git a/Engine/Source/Editor/MainFrame/Private/Frame/MainFrameHandler.h b/Engine/Source/Editor/MainFrame/Private/Frame/MainFrameHandler.h
--- a/Engine/Source/Editor/MainFrame/Private/Frame/MainFrameHandler.h
+++ b/Engine/Source/Editor/MainFrame/Private/Frame/MainFrameHandler.h
@@ -72,7 +72,7 @@ public:
 
 		// Save out any config settings for the editor so they don't get lost
 		GEditor->SaveConfig();
-		GEditorModeTools().SaveConfig();
+		GLevelEditorModeTools().SaveConfig();
 
 		// Delete user settings, if requested
 		if (FUnrealEdMisc::Get().IsDeletePreferences())
@@ -173,10 +173,10 @@ public:
 			bool bOkToExit = true;
 
 			// Check if level Mode is open this does PostEditMove processing on actors when it closes so need to do this first before save dialog
-			if( GEditorModeTools().IsModeActive( FBuiltinEditorModes::EM_Level ) || 
-				GEditorModeTools().IsModeActive( FBuiltinEditorModes::EM_StreamingLevel) )
+			if( GLevelEditorModeTools().IsModeActive( FBuiltinEditorModes::EM_Level ) || 
+				GLevelEditorModeTools().IsModeActive( FBuiltinEditorModes::EM_StreamingLevel) )
 			{
-				GEditorModeTools().ActivateMode( FBuiltinEditorModes::EM_Default );
+				GLevelEditorModeTools().ActivateDefaultMode();
 				bOkToExit = false;
 			}
 
diff --git a/Engine/Source/Editor/MaterialEditor/Private/SMaterialEditorViewport.cpp b/Engine/Source/Editor/MaterialEditor/Private/SMaterialEditorViewport.cpp
--- a/Engine/Source/Editor/MaterialEditor/Private/SMaterialEditorViewport.cpp
+++ b/Engine/Source/Editor/MaterialEditor/Private/SMaterialEditorViewport.cpp
@@ -34,7 +34,7 @@ private:
 };
 
 FMaterialEditorViewportClient::FMaterialEditorViewportClient(TWeakPtr<IMaterialEditor> InMaterialEditor, FPreviewScene& InPreviewScene)
-	: FEditorViewportClient( &InPreviewScene )
+	: FEditorViewportClient( GLevelEditorModeTools(), &InPreviewScene )
 	, MaterialEditorPtr(InMaterialEditor)
 {
 	// Setup defaults for the common draw helper.
diff --git a/Engine/Source/Editor/Matinee/Private/Matinee.cpp b/Engine/Source/Editor/Matinee/Private/Matinee.cpp
--- a/Engine/Source/Editor/Matinee/Private/Matinee.cpp
+++ b/Engine/Source/Editor/Matinee/Private/Matinee.cpp
@@ -308,7 +308,7 @@ void FMatinee::SetAudioRealtimeOverride( bool bAudioIsRealtime ) const
 {
 	for(int32 i=0; i<GEditor->LevelViewportClients.Num(); i++)
 	{
-		FLevelEditorViewportClient* const LevelVC = GEditor->LevelViewportClients[i];
+		FEditorViewportClient* const LevelVC = GEditor->LevelViewportClients[i];
 		if (LevelVC)
 		{
 			if(LevelVC->IsPerspective() && LevelVC->AllowMatineePreview() )
@@ -1998,15 +1998,15 @@ void FMatinee::OnClose()
 	Opt->bAdjustingGroupKeyframes = false;
 
 	// When they close the window - change the mode away from InterpEdit.
-	if( GEditorModeTools().IsModeActive( FBuiltinEditorModes::EM_InterpEdit ) )
+	if( GLevelEditorModeTools().IsModeActive( FBuiltinEditorModes::EM_InterpEdit ) )
 	{
-		FEdModeInterpEdit* InterpEditMode = (FEdModeInterpEdit*)GEditorModeTools().GetActiveMode( FBuiltinEditorModes::EM_InterpEdit );
+		FEdModeInterpEdit* InterpEditMode = (FEdModeInterpEdit*)GLevelEditorModeTools().GetActiveMode( FBuiltinEditorModes::EM_InterpEdit );
 
 		// Only change mode if this window closing wasn't instigated by someone changing mode!
 		if( !InterpEditMode->bLeavingMode )
 		{
 			InterpEditMode->InterpEd = NULL;
-			GEditorModeTools().DeactivateMode( FBuiltinEditorModes::EM_InterpEdit );
+			GLevelEditorModeTools().DeactivateMode( FBuiltinEditorModes::EM_InterpEdit );
 		}
 	}
 
@@ -2131,7 +2131,7 @@ static void DrawTransparentLine( FCanvas* Canvas, const FVector2D& Start, const
 	}
 }
 
-void FMatinee::DrawModeHUD(FLevelEditorViewportClient* ViewportClient,FViewport* Viewport,const FSceneView* View,FCanvas* Canvas)
+void FMatinee::DrawModeHUD(FEditorViewportClient* ViewportClient,FViewport* Viewport,const FSceneView* View,FCanvas* Canvas)
 {
 	if( ViewportClient->AllowMatineePreview() )
 	{
diff --git a/Engine/Source/Editor/Matinee/Private/Matinee.h b/Engine/Source/Editor/Matinee/Private/Matinee.h
--- a/Engine/Source/Editor/Matinee/Private/Matinee.h
+++ b/Engine/Source/Editor/Matinee/Private/Matinee.h
@@ -981,7 +981,7 @@ public:
 	void MoveActiveDown();
 
 	void DrawTracks3D(const FSceneView* View, FPrimitiveDrawInterface* PDI) override;
-	void DrawModeHUD(FLevelEditorViewportClient* ViewportClient,FViewport* Viewport,const FSceneView* View,FCanvas* Canvas) override;
+	void DrawModeHUD(FEditorViewportClient* ViewportClient,FViewport* Viewport,const FSceneView* View,FCanvas* Canvas) override;
 
 	void TickInterp(float DeltaSeconds);
 
diff --git a/Engine/Source/Editor/Matinee/Private/MatineeMenus.cpp b/Engine/Source/Editor/Matinee/Private/MatineeMenus.cpp
--- a/Engine/Source/Editor/Matinee/Private/MatineeMenus.cpp
+++ b/Engine/Source/Editor/Matinee/Private/MatineeMenus.cpp
@@ -4174,16 +4174,16 @@ void FMatinee::OnMenuReduceKeys()
 /** Toggles interting of the panning the interp editor left and right */
 void FMatinee::OnToggleInvertPan()
 {
-	bool bInvertPan = GEditorModeTools().GetInterpPanInvert();
+	bool bInvertPan = GLevelEditorModeTools().GetInterpPanInvert();
 
 	bInvertPan = !bInvertPan;
 
-	GEditorModeTools().SetInterpPanInvert(bInvertPan);
+	GLevelEditorModeTools().SetInterpPanInvert(bInvertPan);
 }
 
 bool FMatinee::IsInvertPanToggled()
 {
-	return GEditorModeTools().GetInterpPanInvert();	
+	return GLevelEditorModeTools().GetInterpPanInvert();	
 }
 
 /** Called when split translation and rotation is selected from a movement track context menu */
diff --git a/Engine/Source/Editor/Matinee/Private/MatineeRecordMovie.cpp b/Engine/Source/Editor/Matinee/Private/MatineeRecordMovie.cpp
--- a/Engine/Source/Editor/Matinee/Private/MatineeRecordMovie.cpp
+++ b/Engine/Source/Editor/Matinee/Private/MatineeRecordMovie.cpp
@@ -208,7 +208,7 @@ FReply SMatineeRecordMovie::OnOK()
 
 	GConfig->Flush( false, GEditorUserSettingsIni );
 
-	FEdModeInterpEdit* Mode = (FEdModeInterpEdit*)GEditorModeTools().GetActiveMode( FBuiltinEditorModes::EM_InterpEdit );
+	FEdModeInterpEdit* Mode = (FEdModeInterpEdit*)GLevelEditorModeTools().GetActiveMode( FBuiltinEditorModes::EM_InterpEdit );
 	if ( Mode != NULL && Mode->InterpEd != NULL )
 	{
 		// Store the options for the capture of the Matinee
@@ -307,7 +307,7 @@ void SMatineeRecordMovie::Construct(const FArguments& InArgs, TWeakPtr<SWindow>
 	ParentWindowPtr = InParentWindow;
 	bUsingCustomResolution = false;
 	FString MatineeName = TEXT("SeqAct_Interp");
-	FEdModeInterpEdit* Mode = (FEdModeInterpEdit*)GEditorModeTools().GetActiveMode( FBuiltinEditorModes::EM_InterpEdit );
+	FEdModeInterpEdit* Mode = (FEdModeInterpEdit*)GLevelEditorModeTools().GetActiveMode( FBuiltinEditorModes::EM_InterpEdit );
 	if ( Mode != NULL && Mode->InterpEd != NULL )
 	{
 		MatineeName = Mode->InterpEd->GetMatineeActor()->GetName();
diff --git a/Engine/Source/Editor/Matinee/Private/MatineeTrackHelpers.cpp b/Engine/Source/Editor/Matinee/Private/MatineeTrackHelpers.cpp
--- a/Engine/Source/Editor/Matinee/Private/MatineeTrackHelpers.cpp
+++ b/Engine/Source/Editor/Matinee/Private/MatineeTrackHelpers.cpp
@@ -34,7 +34,7 @@ UInterpTrackHelper::UInterpTrackHelper(const class FPostConstructInitializePrope
 
 AActor* UInterpTrackHelper::GetGroupActor(const UInterpTrack* Track) const
 {
-	FEdModeInterpEdit* mode = (FEdModeInterpEdit*)GEditorModeTools().GetActiveMode( FBuiltinEditorModes::EM_InterpEdit );
+	FEdModeInterpEdit* mode = (FEdModeInterpEdit*)GLevelEditorModeTools().GetActiveMode( FBuiltinEditorModes::EM_InterpEdit );
 	check(mode != NULL);
 
 	IMatineeBase* InterpEd = mode->InterpEd;
@@ -82,7 +82,7 @@ bool UMatineeTrackAnimControlHelper::PreCreateTrack( UInterpGroup* Group, const
 	// For AnimControl tracks - pop up a dialog to choose slot name.
 	AnimSlotName = NAME_None;
 
-	FEdModeInterpEdit* Mode = (FEdModeInterpEdit*)GEditorModeTools().GetActiveMode( FBuiltinEditorModes::EM_InterpEdit );
+	FEdModeInterpEdit* Mode = (FEdModeInterpEdit*)GLevelEditorModeTools().GetActiveMode( FBuiltinEditorModes::EM_InterpEdit );
 	check(Mode != NULL);
 
 	IMatineeBase* InterpEd = Mode->InterpEd;
@@ -239,7 +239,7 @@ bool UMatineeTrackAnimControlHelper::PreCreateKeyframe( UInterpTrack *Track, flo
 	if ( Skeleton )
 	{
 		// Show the dialog.
-		FEdModeInterpEdit* Mode = (FEdModeInterpEdit*)GEditorModeTools().GetActiveMode( FBuiltinEditorModes::EM_InterpEdit );
+		FEdModeInterpEdit* Mode = (FEdModeInterpEdit*)GLevelEditorModeTools().GetActiveMode( FBuiltinEditorModes::EM_InterpEdit );
 		check(Mode != NULL);
 		check(Mode->InterpEd != NULL);
 		
@@ -323,7 +323,7 @@ bool UMatineeTrackDirectorHelper::PreCreateKeyframe( UInterpTrack *Track, float
 	// If adding a cut, bring up combo to let user choose group to cut to.
 	KeyframeAddDataName = NAME_None;
 
-	FEdModeInterpEdit* Mode = (FEdModeInterpEdit*)GEditorModeTools().GetActiveMode( FBuiltinEditorModes::EM_InterpEdit );
+	FEdModeInterpEdit* Mode = (FEdModeInterpEdit*)GLevelEditorModeTools().GetActiveMode( FBuiltinEditorModes::EM_InterpEdit );
 	check(Mode != NULL);
 	check(Mode->InterpEd != NULL);
 
@@ -392,7 +392,7 @@ bool UMatineeTrackEventHelper::PreCreateKeyframe( UInterpTrack *Track, float Key
 	KeyframeAddDataName = NAME_None;
 
 	// Prompt user for name of new event.
-	FEdModeInterpEdit* Mode = (FEdModeInterpEdit*)GEditorModeTools().GetActiveMode( FBuiltinEditorModes::EM_InterpEdit );
+	FEdModeInterpEdit* Mode = (FEdModeInterpEdit*)GLevelEditorModeTools().GetActiveMode( FBuiltinEditorModes::EM_InterpEdit );
 	check(Mode != NULL);
 	check(Mode->InterpEd != NULL);
 	TSharedRef<STextEntryPopup> TextEntryPopup = 
@@ -494,7 +494,7 @@ bool UMatineeTrackFloatPropHelper::PreCreateTrack( UInterpGroup* Group, const UI
 		// For Property tracks - pop up a dialog to choose property name.
 		TrackAddPropName = NAME_None;
 
-		FEdModeInterpEdit* Mode = (FEdModeInterpEdit*)GEditorModeTools().GetActiveMode( FBuiltinEditorModes::EM_InterpEdit );
+		FEdModeInterpEdit* Mode = (FEdModeInterpEdit*)GLevelEditorModeTools().GetActiveMode( FBuiltinEditorModes::EM_InterpEdit );
 		check(Mode != NULL);
 
 		IMatineeBase* InterpEd = Mode->InterpEd;
@@ -632,7 +632,7 @@ bool UMatineeTrackBoolPropHelper::PreCreateTrack( UInterpGroup* Group, const UIn
 		// For Property tracks - pop up a dialog to choose property name.
 		TrackAddPropName = NAME_None;
 
-		FEdModeInterpEdit* Mode = (FEdModeInterpEdit*)GEditorModeTools().GetActiveMode( FBuiltinEditorModes::EM_InterpEdit );
+		FEdModeInterpEdit* Mode = (FEdModeInterpEdit*)GLevelEditorModeTools().GetActiveMode( FBuiltinEditorModes::EM_InterpEdit );
 		check(Mode != NULL);
 
 		IMatineeBase* InterpEd = Mode->InterpEd;
@@ -767,7 +767,7 @@ bool UMatineeTrackToggleHelper::PreCreateKeyframe( UInterpTrack *Track, float Ke
 
 	bool bResult = false;
 
-	FEdModeInterpEdit* Mode = (FEdModeInterpEdit*)GEditorModeTools().GetActiveMode( FBuiltinEditorModes::EM_InterpEdit );
+	FEdModeInterpEdit* Mode = (FEdModeInterpEdit*)GLevelEditorModeTools().GetActiveMode( FBuiltinEditorModes::EM_InterpEdit );
 	check(Mode != NULL);
 
 	IMatineeBase* InterpEd = Mode->InterpEd;
@@ -844,7 +844,7 @@ bool UMatineeTrackVectorPropHelper::ChooseProperty(TArray<FName> &PropNames) con
 {
 	bool bResult = false;
 
-	FEdModeInterpEdit* Mode = (FEdModeInterpEdit*)GEditorModeTools().GetActiveMode( FBuiltinEditorModes::EM_InterpEdit );
+	FEdModeInterpEdit* Mode = (FEdModeInterpEdit*)GLevelEditorModeTools().GetActiveMode( FBuiltinEditorModes::EM_InterpEdit );
 	check(Mode != NULL);
 
 	IMatineeBase* InterpEd = Mode->InterpEd;
@@ -940,7 +940,7 @@ bool UMatineeTrackVectorPropHelper::PreCreateTrack( UInterpGroup* Group, const U
 		// For Property tracks - pop up a dialog to choose property name.
 		TrackAddPropName = NAME_None;
 
-		FEdModeInterpEdit* Mode = (FEdModeInterpEdit*)GEditorModeTools().GetActiveMode( FBuiltinEditorModes::EM_InterpEdit );
+		FEdModeInterpEdit* Mode = (FEdModeInterpEdit*)GLevelEditorModeTools().GetActiveMode( FBuiltinEditorModes::EM_InterpEdit );
 		check(Mode != NULL);
 
 		IMatineeBase* InterpEd = Mode->InterpEd;
@@ -998,7 +998,7 @@ bool UMatineeTrackColorPropHelper::PreCreateTrack( UInterpGroup* Group, const UI
 		// For Property tracks - pop up a dialog to choose property name.
 		TrackAddPropName = NAME_None;
 
-		FEdModeInterpEdit* Mode = (FEdModeInterpEdit*)GEditorModeTools().GetActiveMode( FBuiltinEditorModes::EM_InterpEdit );
+		FEdModeInterpEdit* Mode = (FEdModeInterpEdit*)GLevelEditorModeTools().GetActiveMode( FBuiltinEditorModes::EM_InterpEdit );
 		check(Mode != NULL);
 
 		IMatineeBase* InterpEd = Mode->InterpEd;
@@ -1063,7 +1063,7 @@ bool UMatineeTrackLinearColorPropHelper::PreCreateTrack( UInterpGroup* Group, co
 		// For Property tracks - pop up a dialog to choose property name.
 		TrackAddPropName = NAME_None;
 
-		FEdModeInterpEdit* Mode = (FEdModeInterpEdit*)GEditorModeTools().GetActiveMode( FBuiltinEditorModes::EM_InterpEdit );
+		FEdModeInterpEdit* Mode = (FEdModeInterpEdit*)GLevelEditorModeTools().GetActiveMode( FBuiltinEditorModes::EM_InterpEdit );
 		check(Mode != NULL);
 
 		IMatineeBase *InterpEd = Mode->InterpEd;
@@ -1120,7 +1120,7 @@ bool UMatineeTrackVisibilityHelper::PreCreateKeyframe( UInterpTrack *Track, floa
 {
 	bool bResult = false;
 
-	FEdModeInterpEdit* Mode = (FEdModeInterpEdit*)GEditorModeTools().GetActiveMode( FBuiltinEditorModes::EM_InterpEdit );
+	FEdModeInterpEdit* Mode = (FEdModeInterpEdit*)GLevelEditorModeTools().GetActiveMode( FBuiltinEditorModes::EM_InterpEdit );
 	check(Mode != NULL);
 
 	IMatineeBase *InterpEd = Mode->InterpEd;
diff --git a/Engine/Source/Editor/Matinee/Private/MatineeViewportClient.cpp b/Engine/Source/Editor/Matinee/Private/MatineeViewportClient.cpp
--- a/Engine/Source/Editor/Matinee/Private/MatineeViewportClient.cpp
+++ b/Engine/Source/Editor/Matinee/Private/MatineeViewportClient.cpp
@@ -14,6 +14,7 @@
 -----------------------------------------------------------------------------*/
 
 FMatineeViewportClient::FMatineeViewportClient( class FMatinee* InMatinee )
+	: FEditorViewportClient(GLevelEditorModeTools())
 {
 	InterpEd = InMatinee;
 
@@ -915,7 +916,7 @@ void FMatineeViewportClient::MouseMove(FViewport* Viewport, int32 X, int32 Y)
 		}
 		else if(bPanning)
 		{
-			const bool bInvertPanning = GEditorModeTools().GetInterpPanInvert();
+			const bool bInvertPanning = GLevelEditorModeTools().GetInterpPanInvert();
 
 			float DeltaTime = (bInvertPanning ? -DeltaX : DeltaX) / InterpEd->PixelsPerSec;
 			InterpEd->ViewStartTime -= DeltaTime;
diff --git a/Engine/Source/Editor/Matinee/Public/IMatinee.h b/Engine/Source/Editor/Matinee/Public/IMatinee.h
--- a/Engine/Source/Editor/Matinee/Public/IMatinee.h
+++ b/Engine/Source/Editor/Matinee/Public/IMatinee.h
@@ -35,7 +35,7 @@ public:
 	virtual void ClearKeySelection() = 0;
 
 	virtual void DrawTracks3D(const FSceneView* View, FPrimitiveDrawInterface* PDI) = 0;
-	virtual void DrawModeHUD(FLevelEditorViewportClient* ViewportClient,FViewport* Viewport,const FSceneView* View,FCanvas* Canvas) = 0;
+	virtual void DrawModeHUD(FEditorViewportClient* ViewportClient,FViewport* Viewport,const FSceneView* View,FCanvas* Canvas) = 0;
 
 	virtual void BeginDrag3DHandle(UInterpGroup* Group, int32 TrackIndex) = 0;
 	virtual void Move3DHandle(UInterpGroup* Group, int32 TrackIndex, int32 KeyIndex, bool bArriving, const FVector& Delta) = 0;
diff --git a/Engine/Source/Editor/MeshPaint/Private/MeshPaintEdMode.cpp b/Engine/Source/Editor/MeshPaint/Private/MeshPaintEdMode.cpp
--- a/Engine/Source/Editor/MeshPaint/Private/MeshPaintEdMode.cpp
+++ b/Engine/Source/Editor/MeshPaint/Private/MeshPaintEdMode.cpp
@@ -92,11 +92,6 @@ FEdModeMeshPaint::FEdModeMeshPaint()
 	  SeamMaskRenderTargetTexture( NULL ),
 	  ScopedTransaction( NULL )
 {
-	ID = FBuiltinEditorModes::EM_MeshPaint;
-	Name = LOCTEXT("MeshPaint_ModeName", "Paint");
-	IconBrush = FSlateIcon(FEditorStyle::GetStyleSetName(), "LevelEditor.MeshPaintMode", "LevelEditor.MeshPaintMode.Small");
-	bVisible = true;
-	PriorityOrder = 200;
 }
 
 /** Destructor */
@@ -171,7 +166,7 @@ void FEdModeMeshPaint::Enter()
 	// Only alter level editor viewports.
 	for( int32 ViewIndex = 0 ; ViewIndex < GEditor->LevelViewportClients.Num() ; ++ViewIndex )
 	{
-		FLevelEditorViewportClient* ViewportClient = GEditor->LevelViewportClients[ViewIndex];
+		FEditorViewportClient* ViewportClient = GEditor->LevelViewportClients[ViewIndex];
 		SetViewportShowFlags( bAllowColorViewModes, *ViewportClient );
 	} 
 
@@ -208,7 +203,7 @@ void FEdModeMeshPaint::Exit()
 	// Only alter level editor viewports.
 	for( int32 ViewIndex = 0 ; ViewIndex < GEditor->LevelViewportClients.Num() ; ++ViewIndex )
 	{
-		FLevelEditorViewportClient* ViewportClient = GEditor->LevelViewportClients[ViewIndex];
+		FEditorViewportClient* ViewportClient = GEditor->LevelViewportClients[ViewIndex];
 		SetViewportShowFlags( bAllowColorViewModes, *ViewportClient );
 	} 
 
@@ -251,7 +246,7 @@ void FEdModeMeshPaint::Exit()
 
 
 /** FEdMode: Called when the mouse is moved over the viewport */
-bool FEdModeMeshPaint::MouseMove( FLevelEditorViewportClient* ViewportClient, FViewport* Viewport, int32 x, int32 y )
+bool FEdModeMeshPaint::MouseMove( FEditorViewportClient* ViewportClient, FViewport* Viewport, int32 x, int32 y )
 {
 	// We only care about perspective viewports
 	if( ViewportClient->IsPerspective() )
@@ -274,7 +269,7 @@ bool FEdModeMeshPaint::MouseMove( FLevelEditorViewportClient* ViewportClient, FV
  *
  * @return	true if input was handled
  */
-bool FEdModeMeshPaint::CapturedMouseMove( FLevelEditorViewportClient* InViewportClient, FViewport* InViewport, int32 InMouseX, int32 InMouseY )
+bool FEdModeMeshPaint::CapturedMouseMove( FEditorViewportClient* InViewportClient, FViewport* InViewport, int32 InMouseX, int32 InMouseY )
 {
 	// We only care about perspective viewports
 	if( InViewportClient->IsPerspective() && InViewportClient->EngineShowFlags.ModeWidgets )
@@ -288,7 +283,7 @@ bool FEdModeMeshPaint::CapturedMouseMove( FLevelEditorViewportClient* InViewport
 				InViewportClient->EngineShowFlags)
 				.SetRealtimeUpdate( InViewportClient->IsRealtime() ));
 			FSceneView* View = InViewportClient->CalcSceneView( &ViewFamily );
-			FViewportCursorLocation MouseViewportRay( View, (FLevelEditorViewportClient*)InViewport->GetClient(), InMouseX, InMouseY );
+			FViewportCursorLocation MouseViewportRay( View, (FEditorViewportClient*)InViewport->GetClient(), InMouseX, InMouseY );
 
 			
 			// Paint!
@@ -311,7 +306,7 @@ bool FEdModeMeshPaint::CapturedMouseMove( FLevelEditorViewportClient* InViewport
 
 
 /** FEdMode: Called when a mouse button is pressed */
-bool FEdModeMeshPaint::StartTracking(FLevelEditorViewportClient* InViewportClient, FViewport* InViewport)
+bool FEdModeMeshPaint::StartTracking(FEditorViewportClient* InViewportClient, FViewport* InViewport)
 {
 	return true;
 }
@@ -319,7 +314,7 @@ bool FEdModeMeshPaint::StartTracking(FLevelEditorViewportClient* InViewportClien
 
 
 /** FEdMode: Called when the a mouse button is released */
-bool FEdModeMeshPaint::EndTracking(FLevelEditorViewportClient* InViewportClient, FViewport* InViewport)
+bool FEdModeMeshPaint::EndTracking(FEditorViewportClient* InViewportClient, FViewport* InViewport)
 {
 	EndPainting();
 	return true;
@@ -368,7 +363,7 @@ void FEdModeMeshPaint::EndPainting()
 }
 
 /** FEdMode: Called when a key is pressed */
-bool FEdModeMeshPaint::InputKey( FLevelEditorViewportClient* InViewportClient, FViewport* InViewport, FKey InKey, EInputEvent InEvent )
+bool FEdModeMeshPaint::InputKey( FEditorViewportClient* InViewportClient, FViewport* InViewport, FKey InKey, EInputEvent InEvent )
 {
 	bool bHandled = false;
 
@@ -477,7 +472,7 @@ bool FEdModeMeshPaint::InputKey( FLevelEditorViewportClient* InViewportClient, F
 							.SetRealtimeUpdate( InViewportClient->IsRealtime() ));
 
 						FSceneView* View = InViewportClient->CalcSceneView( &ViewFamily );
-						FViewportCursorLocation MouseViewportRay( View, (FLevelEditorViewportClient*)InViewport->GetClient(), InViewport->GetMouseX(), InViewport->GetMouseY() );
+						FViewportCursorLocation MouseViewportRay( View, (FEditorViewportClient*)InViewport->GetClient(), InViewport->GetMouseX(), InViewport->GetMouseY() );
 
 						// Paint!
 						const bool bVisualCueOnly = false;
@@ -2687,7 +2682,7 @@ void FEdModeMeshPaint::FinishPaintingTexture( )
 
 
 /** FEdMode: Called when mouse drag input it applied */
-bool FEdModeMeshPaint::InputDelta( FLevelEditorViewportClient* InViewportClient, FViewport* InViewport, FVector& InDrag, FRotator& InRot, FVector& InScale )
+bool FEdModeMeshPaint::InputDelta( FEditorViewportClient* InViewportClient, FViewport* InViewport, FVector& InDrag, FRotator& InRot, FVector& InScale )
 {
 	// We only care about perspective viewports
 	if( InViewportClient->IsPerspective() )
@@ -2719,7 +2714,7 @@ void FEdModeMeshPaint::Render( const FSceneView* View, FViewport* Viewport, FPri
 	FEdMode::Render( View, Viewport, PDI );
 
 	// If this viewport does not support Mode widgets we will not draw it here.
-	FLevelEditorViewportClient* ViewportClient = (FLevelEditorViewportClient*)Viewport->GetClient();
+	FEditorViewportClient* ViewportClient = (FEditorViewportClient*)Viewport->GetClient();
 	if( ViewportClient && !ViewportClient->EngineShowFlags.ModeWidgets)
 	{
 		return;
@@ -3042,7 +3037,7 @@ void FEdModeMeshPaint::ForceRealTimeViewports( const bool bEnable, const bool bS
 	TSharedPtr< ILevelViewport > ViewportWindow = LevelEditorModule.GetFirstActiveViewport();
 	if (ViewportWindow.IsValid())
 	{
-		FLevelEditorViewportClient &Viewport = ViewportWindow->GetLevelViewportClient();
+		FEditorViewportClient &Viewport = ViewportWindow->GetLevelViewportClient();
 		if( Viewport.IsPerspective() )
 		{				
 			if( bEnable )
@@ -3060,7 +3055,7 @@ void FEdModeMeshPaint::ForceRealTimeViewports( const bool bEnable, const bool bS
 }
 
 /** Sets show flags for perspective viewports */
-void FEdModeMeshPaint::SetViewportShowFlags( const bool bAllowColorViewModes, FLevelEditorViewportClient& Viewport )
+void FEdModeMeshPaint::SetViewportShowFlags( const bool bAllowColorViewModes, FEditorViewportClient& Viewport )
 {
 	if( Viewport.IsPerspective() )
 	{	
@@ -4906,7 +4901,7 @@ void FEdModeMeshPaint::BeginTransaction(const FText& Description)
 }
 
 /** FEdModeMeshPaint: Called once per frame */
-void FEdModeMeshPaint::Tick(FLevelEditorViewportClient* ViewportClient,float DeltaTime)
+void FEdModeMeshPaint::Tick(FEditorViewportClient* ViewportClient,float DeltaTime)
 {
 	FEdMode::Tick(ViewportClient,DeltaTime);
 
diff --git a/Engine/Source/Editor/MeshPaint/Private/MeshPaintModule.cpp b/Engine/Source/Editor/MeshPaint/Private/MeshPaintModule.cpp
--- a/Engine/Source/Editor/MeshPaint/Private/MeshPaintModule.cpp
+++ b/Engine/Source/Editor/MeshPaint/Private/MeshPaintModule.cpp
@@ -5,8 +5,6 @@
 
 class FMeshPaintModule : public IMeshPaintModule
 {
-private:
-	TSharedPtr<FEdModeMeshPaint> EdModeMeshPaint;
 public:
 	
 	/**
@@ -14,9 +12,12 @@ public:
 	 */
 	virtual void StartupModule() override
 	{
-		TSharedRef<FEdModeMeshPaint> NewEditorMode = MakeShareable(new FEdModeMeshPaint);
-		GEditorModeTools().RegisterMode(NewEditorMode);
-		EdModeMeshPaint = NewEditorMode;
+		FEditorModeRegistry::Get().RegisterMode<FEdModeMeshPaint>(
+			FBuiltinEditorModes::EM_MeshPaint,
+			NSLOCTEXT("MeshPaint_Mode", "MeshPaint_ModeName", "Paint"),
+			FSlateIcon(FEditorStyle::GetStyleSetName(), "LevelEditor.MeshPaintMode", "LevelEditor.MeshPaintMode.Small"),
+			true, 200
+			);
 	}
 
 	/**
@@ -24,8 +25,7 @@ public:
 	 */
 	virtual void ShutdownModule() override
 	{
-		GEditorModeTools().UnregisterMode(EdModeMeshPaint.ToSharedRef());
-		EdModeMeshPaint = NULL;
+		FEditorModeRegistry::Get().UnregisterMode(FBuiltinEditorModes::EM_MeshPaint);
 	}
 };
 
diff --git a/Engine/Source/Editor/MeshPaint/Private/SMeshPaint.cpp b/Engine/Source/Editor/MeshPaint/Private/SMeshPaint.cpp
--- a/Engine/Source/Editor/MeshPaint/Private/SMeshPaint.cpp
+++ b/Engine/Source/Editor/MeshPaint/Private/SMeshPaint.cpp
@@ -1088,7 +1088,7 @@ FText FMeshPaintToolKit::GetBaseToolkitName() const
 
 class FEdMode* FMeshPaintToolKit::GetEditorMode() const
 {
-	return GEditorModeTools().GetActiveMode(FBuiltinEditorModes::EM_MeshPaint);
+	return GLevelEditorModeTools().GetActiveMode(FBuiltinEditorModes::EM_MeshPaint);
 }
 
 TSharedPtr<SWidget> FMeshPaintToolKit::GetInlineContent() const
@@ -1096,10 +1096,15 @@ TSharedPtr<SWidget> FMeshPaintToolKit::GetInlineContent() const
 	return MeshPaintWidgets;
 }
 
+SMeshPaint::~SMeshPaint()
+{
+	MeshPaintEditMode = nullptr;
+}
+
 BEGIN_SLATE_FUNCTION_BUILD_OPTIMIZATION
 void SMeshPaint::Construct(const FArguments& InArgs, TSharedRef<FMeshPaintToolKit> InParentToolkit)
 {
-	MeshPaintEditMode = (FEdModeMeshPaint*)GEditorModeTools().GetActiveMode( FBuiltinEditorModes::EM_MeshPaint );
+	MeshPaintEditMode = (FEdModeMeshPaint*)GLevelEditorModeTools().GetActiveMode( FBuiltinEditorModes::EM_MeshPaint );
 
 	PaintColorSet = EMeshPaintColorSet::PaintColor;
 	bShowImportOptions = false;
@@ -2089,7 +2094,7 @@ void SMeshPaint::Tick( const FGeometry& AllottedGeometry, const double InCurrent
 
 class FEdMode* SMeshPaint::GetEditorMode() const
 {
-	return (FEdModeMeshPaint*)GEditorModeTools().GetActiveMode(FBuiltinEditorModes::EM_MeshPaint);
+	return (FEdModeMeshPaint*)GLevelEditorModeTools().GetActiveMode(FBuiltinEditorModes::EM_MeshPaint);
 }
 
 EVisibility SMeshPaint::GetResourceTypeVerticesVisibility() const
diff --git a/Engine/Source/Editor/MeshPaint/Private/SMeshPaint.h b/Engine/Source/Editor/MeshPaint/Private/SMeshPaint.h
--- a/Engine/Source/Editor/MeshPaint/Private/SMeshPaint.h
+++ b/Engine/Source/Editor/MeshPaint/Private/SMeshPaint.h
@@ -65,6 +65,8 @@ public:
 	SLATE_BEGIN_ARGS( SMeshPaint ){}
 	SLATE_END_ARGS()
 
+	~SMeshPaint();
+
 	/** SCompoundWidget functions */
 	void Construct(const FArguments& InArgs, TSharedRef<FMeshPaintToolKit> InParentToolkit);
 
diff --git a/Engine/Source/Editor/MeshPaint/Public/MeshPaintEdMode.h b/Engine/Source/Editor/MeshPaint/Public/MeshPaintEdMode.h
--- a/Engine/Source/Editor/MeshPaint/Public/MeshPaintEdMode.h
+++ b/Engine/Source/Editor/MeshPaint/Public/MeshPaintEdMode.h
@@ -359,12 +359,12 @@ public:
 	virtual bool UsesToolkits() const override;
 	virtual void Enter() override;
 	virtual void Exit() override;
-	virtual bool MouseMove( FLevelEditorViewportClient* ViewportClient, FViewport* Viewport, int32 x, int32 y ) override;
-	virtual bool CapturedMouseMove( FLevelEditorViewportClient* InViewportClient, FViewport* InViewport, int32 InMouseX, int32 InMouseY ) override;
-	virtual bool StartTracking(FLevelEditorViewportClient* InViewportClient, FViewport* InViewport) override;
-	virtual bool EndTracking(FLevelEditorViewportClient* InViewportClient, FViewport* InViewport) override;
-	virtual bool InputKey( FLevelEditorViewportClient* InViewportClient, FViewport* InViewport, FKey InKey, EInputEvent InEvent ) override;
-	virtual bool InputDelta( FLevelEditorViewportClient* InViewportClient, FViewport* InViewport, FVector& InDrag, FRotator& InRot, FVector& InScale ) override;
+	virtual bool MouseMove( FEditorViewportClient* ViewportClient, FViewport* Viewport, int32 x, int32 y ) override;
+	virtual bool CapturedMouseMove( FEditorViewportClient* InViewportClient, FViewport* InViewport, int32 InMouseX, int32 InMouseY ) override;
+	virtual bool StartTracking(FEditorViewportClient* InViewportClient, FViewport* InViewport) override;
+	virtual bool EndTracking(FEditorViewportClient* InViewportClient, FViewport* InViewport) override;
+	virtual bool InputKey( FEditorViewportClient* InViewportClient, FViewport* InViewport, FKey InKey, EInputEvent InEvent ) override;
+	virtual bool InputDelta( FEditorViewportClient* InViewportClient, FViewport* InViewport, FVector& InDrag, FRotator& InRot, FVector& InScale ) override;
 	virtual void PostUndo() override;
 	virtual void Render( const FSceneView* View, FViewport* Viewport, FPrimitiveDrawInterface* PDI ) override;
 	virtual bool Select( AActor* InActor, bool bInSelected ) override;
@@ -373,7 +373,7 @@ public:
 	virtual bool AllowWidgetMove() override { return false; }
 	virtual bool ShouldDrawWidget() const override { return false; }
 	virtual bool UsesTransformWidget() const override { return false; }
-	virtual void Tick( FLevelEditorViewportClient* ViewportClient, float DeltaTime ) override;
+	virtual void Tick( FEditorViewportClient* ViewportClient, float DeltaTime ) override;
 	// End of FEdMode interface
 
 
@@ -637,7 +637,7 @@ private:
 	void ForceRealTimeViewports( const bool bEnable, const bool bStoreCurrentState );
 
 	/** Sets show flags for perspective viewports */
-	void SetViewportShowFlags( const bool bAllowColorViewModes, FLevelEditorViewportClient& Viewport );
+	void SetViewportShowFlags( const bool bAllowColorViewModes, FEditorViewportClient& Viewport );
 
 	/** Starts painting a texture */
 	void StartPaintingTexture( UStaticMeshComponent* InStaticMeshComponent );
diff --git a/Engine/Source/Editor/Persona/Private/AnimationEditorViewportClient.cpp b/Engine/Source/Editor/Persona/Private/AnimationEditorViewportClient.cpp
--- a/Engine/Source/Editor/Persona/Private/AnimationEditorViewportClient.cpp
+++ b/Engine/Source/Editor/Persona/Private/AnimationEditorViewportClient.cpp
@@ -67,7 +67,7 @@ IMPLEMENT_HIT_PROXY( HPersonaBoneProxy, HHitProxy );
 // FAnimationViewportClient
 
 FAnimationViewportClient::FAnimationViewportClient( FPreviewScene& InPreviewScene, TWeakPtr<FPersona> InPersonaPtr )
-	: FEditorViewportClient(&InPreviewScene)
+	: FEditorViewportClient(GLevelEditorModeTools(), &InPreviewScene)
 	, PersonaPtr( InPersonaPtr )
 	, bManipulating(false)
 	, bInTransaction(false)
@@ -1236,12 +1236,12 @@ FMatrix FAnimationViewportClient::GetWidgetCoordSystem() const
 
 ECoordSystem FAnimationViewportClient::GetWidgetCoordSystemSpace() const
 { 
-	return GEditorModeTools().GetCoordSystem();
+	return GLevelEditorModeTools().GetCoordSystem();
 }
 
 void FAnimationViewportClient::SetWidgetCoordSystemSpace(ECoordSystem NewCoordSystem)
 {
-	GEditorModeTools().SetCoordSystem(NewCoordSystem);
+	GLevelEditorModeTools().SetCoordSystem(NewCoordSystem);
 	Invalidate();
 }
 
diff --git a/Engine/Source/Editor/PhAT/Private/PhATPreviewViewportClient.cpp b/Engine/Source/Editor/PhAT/Private/PhATPreviewViewportClient.cpp
--- a/Engine/Source/Editor/PhAT/Private/PhATPreviewViewportClient.cpp
+++ b/Engine/Source/Editor/PhAT/Private/PhATPreviewViewportClient.cpp
@@ -13,7 +13,7 @@
 
 
 FPhATEdPreviewViewportClient::FPhATEdPreviewViewportClient(TWeakPtr<FPhAT> InPhAT, TSharedPtr<FPhATSharedData> Data)
-	: FEditorViewportClient(&Data->PreviewScene)
+	: FEditorViewportClient(GLevelEditorModeTools(), &Data->PreviewScene)
 	, PhATPtr(InPhAT)
 	, SharedData(Data)
 	, MinPrimSize(0.5f)
diff --git a/Engine/Source/Editor/PlacementMode/Private/PlacementMode.cpp b/Engine/Source/Editor/PlacementMode/Private/PlacementMode.cpp
--- a/Engine/Source/Editor/PlacementMode/Private/PlacementMode.cpp
+++ b/Engine/Source/Editor/PlacementMode/Private/PlacementMode.cpp
@@ -2,6 +2,7 @@
 
 #include "PlacementModePrivatePCH.h"
 
+#include "IPlacementModeModule.h"
 #include "PlacementMode.h"
 #include "PlacementModeToolkit.h"
 
@@ -12,56 +13,23 @@
 #include "Editor/LevelEditor/Public/LevelEditor.h"
 #include "Editor/LevelEditor/Public/ILevelViewport.h"
 
-TSharedRef< FPlacementMode > FPlacementMode::Create()
-{
-	TSharedRef< FPlacementMode > PlacementMode = MakeShareable( new FPlacementMode() );
-	PlacementMode->Initialize();
-	return PlacementMode;
-}
-
 FPlacementMode::FPlacementMode()
 	: PlacementsChanged( false )
 	, CreatedPreviewActors( false )
 	, AssetsToPlace()
 	, ActiveTransactionIndex( INDEX_NONE )
 {
-	ID = TEXT("PLACEMENT");
-	Name = NSLOCTEXT( "PlacementMode", "DisplayName", "Place" );
-	IconBrush = FSlateIcon( FEditorStyle::GetStyleSetName(), "LevelEditor.PlacementMode", "LevelEditor.PlacementMode.Small" );
-	bVisible = true;
-	PriorityOrder = 0;
 }
 
 FPlacementMode::~FPlacementMode()
 {
-	GEditorModeTools().OnEditorModeChanged().RemoveAll( this );
-
-	FAssetRegistryModule& AssetRegistryModule = FModuleManager::GetModuleChecked<FAssetRegistryModule>(TEXT("AssetRegistry"));
-	AssetRegistryModule.Get().OnAssetRemoved().RemoveAll( this );
-	AssetRegistryModule.Get().OnAssetRenamed().RemoveAll( this );
+	
 }
 
 void FPlacementMode::Initialize()
 {
 	//GConfig->GetFloat(TEXT("PlacementMode"), TEXT("AssetThumbnailScale"), AssetThumbnailScale, GEditorUserSettingsIni);
 	//GConfig->GetBool( TEXT( "PlacementMode" ), TEXT( "ShowOtherDeveloperAssets" ), ShowOtherDeveloperAssets, GEditorUserSettingsIni );
-	
-	TArray< FString > RecentlyPlacedAsStrings;
-	GConfig->GetArray(TEXT("PlacementMode"), TEXT("RecentlyPlaced"), RecentlyPlacedAsStrings, GEditorUserSettingsIni);
-
-	//FString ActivePaletteName;
-	//GConfig->GetString( TEXT( "PlacementMode" ), TEXT( "ActivePalette" ), ActivePaletteName, GEditorUserSettingsIni );
-
-	for (int Index = 0; Index < RecentlyPlacedAsStrings.Num(); Index++)
-	{
-		RecentlyPlaced.Add( FActorPlacementInfo( RecentlyPlacedAsStrings[Index] ) );
-	}
-
-	GEditorModeTools().OnEditorModeChanged().AddSP( this, &FPlacementMode::EditorModeChanged );
-
-	FAssetRegistryModule& AssetRegistryModule = FModuleManager::GetModuleChecked<FAssetRegistryModule>(TEXT("AssetRegistry"));
-	AssetRegistryModule.Get().OnAssetRemoved().AddSP( this, &FPlacementMode::OnAssetRemoved );
-	AssetRegistryModule.Get().OnAssetRenamed().AddSP( this, &FPlacementMode::OnAssetRenamed );
 }
 
 bool FPlacementMode::UsesToolkits() const
@@ -69,25 +37,6 @@ bool FPlacementMode::UsesToolkits() const
 	return true;
 }
 
-void FPlacementMode::EditorModeChanged( FEdMode* Mode, bool IsEntering )
-{
-	if ( IsEntering && FModuleManager::Get().IsModuleLoaded( "LevelEditor" ) )
-	{
-		TSharedPtr< ILevelEditor > ToolkitHost = FModuleManager::LoadModuleChecked< FLevelEditorModule >( "LevelEditor" ).GetFirstLevelEditor();
-
-		if ( ToolkitHost.IsValid() )
-		{
-			TArray< FEdMode* > ActiveModes;
-			GEditorModeTools().GetActiveModes( ActiveModes );
-
-			if ( ActiveModes.Num() == 1 && ActiveModes[0]->GetID() == FBuiltinEditorModes::EM_Default )
-			{
-				FLevelEditorActionCallbacks::ExecuteExecCommand( FString( TEXT( "MODE " ) ) + ID.ToString() );
-			}
-		}
-	}
-}
-
 void FPlacementMode::Enter()
 {
 	// Call parent implementation
@@ -113,7 +62,7 @@ void FPlacementMode::Exit()
 	FEdMode::Exit();
 }
 
-void FPlacementMode::Tick(FLevelEditorViewportClient* ViewportClient,float DeltaTime)
+void FPlacementMode::Tick(FEditorViewportClient* ViewportClient,float DeltaTime)
 {
 	if ( IsCurrentlyPlacing() )
 	{
@@ -173,7 +122,7 @@ void FPlacementMode::Tick(FLevelEditorViewportClient* ViewportClient,float Delta
 	FEdMode::Tick( ViewportClient, DeltaTime );
 }
 
-bool FPlacementMode::MouseEnter( FLevelEditorViewportClient* ViewportClient, FViewport* Viewport, int32 x, int32 y )
+bool FPlacementMode::MouseEnter( FEditorViewportClient* ViewportClient, FViewport* Viewport, int32 x, int32 y )
 {
 	if ( IsCurrentlyPlacing() )
 	{
@@ -183,7 +132,7 @@ bool FPlacementMode::MouseEnter( FLevelEditorViewportClient* ViewportClient, FVi
 	return FEdMode::MouseEnter( ViewportClient, Viewport, x, y );
 }
 
-bool FPlacementMode::MouseLeave( FLevelEditorViewportClient* ViewportClient, FViewport* Viewport )
+bool FPlacementMode::MouseLeave( FEditorViewportClient* ViewportClient, FViewport* Viewport )
 {
 	if ( !ViewportClient->IsTracking() )
 	{
@@ -195,18 +144,18 @@ bool FPlacementMode::MouseLeave( FLevelEditorViewportClient* ViewportClient, FVi
 	return FEdMode::MouseLeave( ViewportClient, Viewport );
 }
 
-bool FPlacementMode::AllowPreviewActors( FLevelEditorViewportClient* ViewportClient ) const
+bool FPlacementMode::AllowPreviewActors( FEditorViewportClient* ViewportClient ) const
 {
 	return IsCurrentlyPlacing() && ( ViewportClient->IsTracking() && AllowPreviewActorsWhileTracking ) || !ViewportClient->IsTracking();
 }
 
-bool FPlacementMode::MouseMove( FLevelEditorViewportClient* ViewportClient, FViewport* Viewport, int32 x, int32 y )
+bool FPlacementMode::MouseMove( FEditorViewportClient* ViewportClient, FViewport* Viewport, int32 x, int32 y )
 {
 	UpdatePreviewActors( ViewportClient, Viewport, x, y );
 	return FEdMode::MouseMove( ViewportClient, Viewport, x, y );
 }
 
-void FPlacementMode::UpdatePreviewActors( FLevelEditorViewportClient* ViewportClient, FViewport* Viewport, int32 x, int32 y )
+void FPlacementMode::UpdatePreviewActors( FEditorViewportClient* ViewportClient, FViewport* Viewport, int32 x, int32 y )
 {
 	bool AllAssetsValid = false;
 	bool AllAssetsCanBeDropped = false;
@@ -284,7 +233,7 @@ void FPlacementMode::UpdatePreviewActors( FLevelEditorViewportClient* ViewportCl
 	}
 }
 
-bool FPlacementMode::InputKey( FLevelEditorViewportClient* InViewportClient, FViewport* InViewport, FKey InKey, EInputEvent InEvent )
+bool FPlacementMode::InputKey( FEditorViewportClient* InViewportClient, FViewport* InViewport, FKey InKey, EInputEvent InEvent )
 {
 	bool Handled = false;
 	const bool bIsCtrlDown = ( ( InKey == EKeys::LeftControl || InKey == EKeys::RightControl ) && InEvent != IE_Released ) || InViewport->KeyState( EKeys::LeftControl ) || InViewport->KeyState( EKeys::RightControl );
@@ -309,7 +258,7 @@ bool FPlacementMode::InputKey( FLevelEditorViewportClient* InViewportClient, FVi
 				SelectPlacedActors();
 
 				ClearAssetsToPlace();
-				BroadcastStoppedPlacing(true);
+				IPlacementModeModule::Get().BroadcastStoppedPlacing( true );
 			}
 		}
 	}
@@ -322,7 +271,7 @@ bool FPlacementMode::InputKey( FLevelEditorViewportClient* InViewportClient, FVi
 	return FEdMode::InputKey( InViewportClient, InViewport, InKey, InEvent );
 }
 
-bool FPlacementMode::StartTracking(FLevelEditorViewportClient* InViewportClient, FViewport* InViewport)
+bool FPlacementMode::StartTracking(FEditorViewportClient* InViewportClient, FViewport* InViewport)
 {
 	PlacedActorsThisTrackingSession = false;
 
@@ -336,7 +285,7 @@ bool FPlacementMode::StartTracking(FLevelEditorViewportClient* InViewportClient,
 	return FEdMode::StartTracking(InViewportClient, InViewport);
 }
 
-bool FPlacementMode::EndTracking(FLevelEditorViewportClient* InViewportClient, FViewport* InViewport)
+bool FPlacementMode::EndTracking(FEditorViewportClient* InViewportClient, FViewport* InViewport)
 {
 	if ( IsCurrentlyPlacing() )
 	{
@@ -361,7 +310,7 @@ bool FPlacementMode::EndTracking(FLevelEditorViewportClient* InViewportClient, F
 	return FEdMode::EndTracking(InViewportClient, InViewport);
 }
 
-bool FPlacementMode::HandleClick(FLevelEditorViewportClient* InViewportClient, HHitProxy *HitProxy, const FViewportClick &Click )
+bool FPlacementMode::HandleClick(FEditorViewportClient* InViewportClient, HHitProxy *HitProxy, const FViewportClick &Click )
 {
 	bool Handled = false;
 
@@ -403,7 +352,7 @@ bool FPlacementMode::HandleClick(FLevelEditorViewportClient* InViewportClient, H
 				if ( !Click.IsControlDown() )
 				{
 					ClearAssetsToPlace();
-					BroadcastStoppedPlacing(true);
+					IPlacementModeModule::Get().BroadcastStoppedPlacing( true );
 					InViewportClient->SetRequiredCursorOverride( true, EMouseCursor::GrabHand );
 				}
 
@@ -443,7 +392,7 @@ bool FPlacementMode::HandleClick(FLevelEditorViewportClient* InViewportClient, H
 	return Handled;
 }
 
-bool FPlacementMode::InputDelta( FLevelEditorViewportClient* InViewportClient, FViewport* InViewport, FVector& InDrag, FRotator& InRot, FVector& InScale )
+bool FPlacementMode::InputDelta( FEditorViewportClient* InViewportClient, FViewport* InViewport, FVector& InDrag, FRotator& InRot, FVector& InScale )
 {
 	if ( IsCurrentlyPlacing() )
 	{
@@ -493,6 +442,20 @@ bool FPlacementMode::UsesPropertyWidgets() const
 	return IsCurrentlyPlacing() ? false : FEdMode::ShouldDrawWidget(); 
 }
 
+bool FPlacementMode::IsCompatibleWith(FEditorModeID OtherModeID) const
+{
+	return
+		OtherModeID == FBuiltinEditorModes::EM_Bsp			||
+		OtherModeID == FBuiltinEditorModes::EM_Geometry		||
+		OtherModeID == FBuiltinEditorModes::EM_InterpEdit	||
+		OtherModeID == FBuiltinEditorModes::EM_MeshPaint	||
+		OtherModeID == FBuiltinEditorModes::EM_Landscape	||
+		OtherModeID == FBuiltinEditorModes::EM_Foliage		||
+		OtherModeID == FBuiltinEditorModes::EM_Level		||
+		OtherModeID == FBuiltinEditorModes::EM_Physics		||
+		OtherModeID == FBuiltinEditorModes::EM_ActorPicker;
+}
+
 void FPlacementMode::StartPlacing( const TArray< UObject* >& Assets, UActorFactory* Factory )
 {
 	const bool bNotifySelectNone = true;
@@ -534,7 +497,7 @@ void FPlacementMode::StartPlacing( const TArray< UObject* >& Assets, UActorFacto
 		}
 	}
 
-	StartedPlacingEvent.Broadcast( Assets );
+	IPlacementModeModule::Get().BroadcastStartedPlacing(Assets);
 }
 
 void FPlacementMode::StopPlacing()
@@ -542,7 +505,7 @@ void FPlacementMode::StopPlacing()
 	if ( IsCurrentlyPlacing() )
 	{
 		ClearAssetsToPlace();
-		BroadcastStoppedPlacing(false);
+		IPlacementModeModule::Get().BroadcastStoppedPlacing( false );
 		PlacementsChanged = true;
 	}
 }
@@ -596,11 +559,6 @@ void FPlacementMode::ClearAssetsToPlace()
 	PlacementsChanged = true;
 }
 
-void FPlacementMode::BroadcastStoppedPlacing( bool WasSuccessfullyPlaced ) const
-{
-	StoppedPlacingEvent.Broadcast( WasSuccessfullyPlaced );
-}
-
 void FPlacementMode::SelectPlacedActors()
 {
 	const FScopedTransaction Transaction( NSLOCTEXT( "BuilderMode", "SelectActors", "Select Actors") );
@@ -627,91 +585,6 @@ void FPlacementMode::SelectPlacedActors()
 	GEditor->NoteSelectionChange();
 }
 
-void FPlacementMode::AddToRecentlyPlaced( const TArray< UObject* >& PlacedObjects, UActorFactory* FactoryUsed )
-{
-	FString FactoryPath;
-	if ( FactoryUsed != NULL )
-	{
-		FactoryPath = FactoryUsed->GetPathName();
-	}
-
-	TArray< UObject* > FilteredPlacedObjects;
-	for ( UObject* PlacedObject : PlacedObjects )
-	{
-		// Don't include null placed objects that just have factories.
-		if ( PlacedObject == NULL )
-		{
-			continue;
-		}
-
-		// Don't add brush builders to the recently placed.
-		if ( PlacedObject->IsA(UBrushBuilder::StaticClass()) )
-		{
-			continue;
-		}
-
-		FilteredPlacedObjects.Add(PlacedObject);
-	}
-
-	// Don't change the recently placed if nothing passed the filter.
-	if ( FilteredPlacedObjects.Num() == 0 )
-	{
-		return;
-	}
-
-	bool Changed = false;
-	for ( int Index = 0; Index < FilteredPlacedObjects.Num(); Index++ )
-	{
-		Changed |= RecentlyPlaced.Remove(FActorPlacementInfo(FilteredPlacedObjects[Index]->GetPathName(), FactoryPath)) > 0;
-	}
-
-	for ( int Index = 0; Index < FilteredPlacedObjects.Num(); Index++ )
-	{
-		if ( FilteredPlacedObjects[Index] != NULL )
-		{
-			RecentlyPlaced.Insert(FActorPlacementInfo(FilteredPlacedObjects[Index]->GetPathName(), FactoryPath), 0);
-			Changed = true;
-		}
-	}
-
-	for (int Index = RecentlyPlaced.Num() - 1; Index >= 20; Index--)
-	{
-		RecentlyPlaced.RemoveAt( Index );
-		Changed = true;
-	}
-
-	if ( Changed )
-	{
-		TArray< FString > RecentlyPlacedAsStrings;
-		for (int Index = 0; Index < RecentlyPlaced.Num(); Index++)
-		{
-			RecentlyPlacedAsStrings.Add( RecentlyPlaced[Index].ToString() );
-		}
-
-		GConfig->SetArray(TEXT("PlacementMode"), TEXT("RecentlyPlaced"), RecentlyPlacedAsStrings, GEditorUserSettingsIni);
-		RecentlyPlacedChanged.Broadcast( RecentlyPlaced );
-	}
-}
-
-void FPlacementMode::OnAssetRemoved( const FAssetData& /*InRemovedAssetData*/ )
-{
-	RecentlyPlacedChanged.Broadcast(RecentlyPlaced);
-}
-
-void FPlacementMode::OnAssetRenamed( const FAssetData& AssetData, const FString& OldObjectPath )
-{
-	for (auto& RecentlyPlacedItem : RecentlyPlaced)
-	{
-		if (RecentlyPlacedItem.ObjectPath == OldObjectPath)
-		{
-			RecentlyPlacedItem.ObjectPath = AssetData.ObjectPath.ToString();
-			break;
-		}
-	}
-
-	RecentlyPlacedChanged.Broadcast(RecentlyPlaced);
-}
-
 const TArray< TWeakObjectPtr<UObject> >& FPlacementMode::GetCurrentlyPlacingObjects() const
 {
 	return AssetsToPlace;
diff --git a/Engine/Source/Editor/PlacementMode/Private/PlacementMode.h b/Engine/Source/Editor/PlacementMode/Private/PlacementMode.h
--- a/Engine/Source/Editor/PlacementMode/Private/PlacementMode.h
+++ b/Engine/Source/Editor/PlacementMode/Private/PlacementMode.h
@@ -6,38 +6,34 @@
 class FPlacementMode : public FEdMode, public IPlacementMode
 {
 public:
-	static TSharedRef< class FPlacementMode > Create();
-
+	FPlacementMode();
 	~FPlacementMode();
 
 	// FEdMode interface
 	virtual bool UsesToolkits() const override;
 	void Enter() override;
 	void Exit() override;
-	virtual void Tick(FLevelEditorViewportClient* ViewportClient,float DeltaTime) override;
-	virtual bool MouseEnter( FLevelEditorViewportClient* ViewportClient,FViewport* Viewport,int32 x, int32 y ) override;
-	virtual bool MouseLeave( FLevelEditorViewportClient* ViewportClient,FViewport* Viewport ) override;
-	virtual bool MouseMove( FLevelEditorViewportClient* ViewportClient, FViewport* Viewport, int32 x, int32 y ) override;
-	virtual bool InputKey( FLevelEditorViewportClient* InViewportClient, FViewport* InViewport, FKey InKey, EInputEvent InEvent ) override;
-	virtual bool EndTracking( FLevelEditorViewportClient* InViewportClient, FViewport* InViewport ) override;
-	virtual bool StartTracking( FLevelEditorViewportClient* InViewportClient, FViewport* InViewport ) override;
-
-	virtual bool HandleClick(FLevelEditorViewportClient* InViewportClient, HHitProxy *HitProxy, const FViewportClick &Click ) override;
-	virtual bool InputDelta( FLevelEditorViewportClient* InViewportClient, FViewport* InViewport, FVector& InDrag, FRotator& InRot, FVector& InScale ) override;
+	virtual void Tick(FEditorViewportClient* ViewportClient,float DeltaTime) override;
+	virtual bool MouseEnter( FEditorViewportClient* ViewportClient,FViewport* Viewport,int32 x, int32 y ) override;
+	virtual bool MouseLeave( FEditorViewportClient* ViewportClient,FViewport* Viewport ) override;
+	virtual bool MouseMove( FEditorViewportClient* ViewportClient, FViewport* Viewport, int32 x, int32 y ) override;
+	virtual bool InputKey( FEditorViewportClient* InViewportClient, FViewport* InViewport, FKey InKey, EInputEvent InEvent ) override;
+	virtual bool EndTracking( FEditorViewportClient* InViewportClient, FViewport* InViewport ) override;
+	virtual bool StartTracking( FEditorViewportClient* InViewportClient, FViewport* InViewport ) override;
+
+	virtual bool HandleClick(FEditorViewportClient* InViewportClient, HHitProxy *HitProxy, const FViewportClick &Click ) override;
+	virtual bool InputDelta( FEditorViewportClient* InViewportClient, FViewport* InViewport, FVector& InDrag, FRotator& InRot, FVector& InScale ) override;
 	virtual bool ShouldDrawWidget() const override;
 	virtual bool UsesPropertyWidgets() const override;
+	virtual bool IsCompatibleWith(FEditorModeID OtherModeID) const override;
 	// End of FEdMode interface
 
 public: //IPlacementMode
 
-	DECLARE_DERIVED_EVENT( FPlacementMode, IPlacementMode::FOnStoppedPlacingEvent, FOnStoppedPlacingEvent );
-	virtual FOnStoppedPlacingEvent& OnStoppedPlacing() override { return StoppedPlacingEvent; }
 	virtual void StopPlacing() override;
 
 	virtual bool IsCurrentlyPlacing() const override { return AssetsToPlace.Num() > 0; };
 
-	DECLARE_DERIVED_EVENT( FPlacementMode, IPlacementMode::FOnStartedPlacingEvent, FOnStartedPlacingEvent );
-	virtual FOnStartedPlacingEvent& OnStartedPlacing() override { return StartedPlacingEvent; }
 	virtual void StartPlacing( const TArray< UObject* >& Assets, UActorFactory* Factory = NULL ) override;
 
 	virtual UActorFactory* GetPlacingFactory() const override { return PlacementFactory.Get(); }
@@ -47,32 +43,18 @@ public: //IPlacementMode
 	virtual void AddValidFocusTargetForPlacement( const TWeakPtr< SWidget >& Widget ) override { ValidFocusTargetsForPlacement.Add( Widget ); }
 	virtual void RemoveValidFocusTargetForPlacement( const TWeakPtr< SWidget >& Widget ) override { ValidFocusTargetsForPlacement.Remove( Widget ); }
 
-	DECLARE_DERIVED_EVENT( FPlacementMode, IPlacementMode::FOnRecentlyPlacedChanged, FOnRecentlyPlacedChanged );
-	virtual FOnRecentlyPlacedChanged& OnRecentlyPlacedChanged() override { return RecentlyPlacedChanged; }
-	virtual void AddToRecentlyPlaced( const TArray< UObject* >& PlacedObjects, UActorFactory* FactoryUsed = NULL ) override;
-	virtual const TArray< FActorPlacementInfo >& GetRecentlyPlaced() const override { return RecentlyPlaced; }
-
 	virtual const TArray< TWeakObjectPtr<UObject> >& GetCurrentlyPlacingObjects() const override;
 
 private:
-
-	FPlacementMode();
-	void Initialize();
+	virtual void Initialize() override;
 
 	void ClearAssetsToPlace();
 
 	void SelectPlacedActors();
 
-	void BroadcastStoppedPlacing( bool WasSuccessfullyPlaced ) const;
+	bool AllowPreviewActors( FEditorViewportClient* ViewportClient ) const;
 
-	bool AllowPreviewActors( FLevelEditorViewportClient* ViewportClient ) const;
-
-	void UpdatePreviewActors( FLevelEditorViewportClient* ViewportClient, FViewport* Viewport, int32 x, int32 y );
-
-	void EditorModeChanged( FEdMode* Mode, bool IsEntering );
-
-	void OnAssetRemoved( const FAssetData& InRemovedAssetData );
-	void OnAssetRenamed( const FAssetData& AssetData, const FString& OldObjectPath );
+	void UpdatePreviewActors( FEditorViewportClient* ViewportClient, FViewport* Viewport, int32 x, int32 y );
 
 private:
 
@@ -82,9 +64,6 @@ private:
 	TArray< TWeakObjectPtr< UObject > > AssetsToPlace;
 	TArray< TWeakObjectPtr< AActor > > PlacedActors;
 
-	FOnStoppedPlacingEvent StoppedPlacingEvent;
-	FOnStartedPlacingEvent StartedPlacingEvent;
-
 	int32 ActiveTransactionIndex;
 
 	bool PlacementsChanged;
@@ -92,8 +71,5 @@ private:
 	bool PlacedActorsThisTrackingSession;
 	bool AllowPreviewActorsWhileTracking;
 
-	FOnRecentlyPlacedChanged RecentlyPlacedChanged;
-	TArray< FActorPlacementInfo > RecentlyPlaced;
-
 	TMap< FName, TWeakObjectPtr< UActorFactory > > AssetTypeToFactory;
 };
diff --git a/Engine/Source/Editor/PlacementMode/Private/PlacementModeModule.cpp b/Engine/Source/Editor/PlacementMode/Private/PlacementModeModule.cpp
--- a/Engine/Source/Editor/PlacementMode/Private/PlacementModeModule.cpp
+++ b/Engine/Source/Editor/PlacementMode/Private/PlacementModeModule.cpp
@@ -13,17 +13,28 @@ public:
 	 */
 	virtual void StartupModule() override
 	{
-		PlacementMode = FPlacementMode::Create();
-		GEditorModeTools().RegisterMode( PlacementMode.ToSharedRef() );
-		GEditorModeTools().RegisterCompatibleModes(PlacementMode->GetID(), FBuiltinEditorModes::EM_Bsp);
-		GEditorModeTools().RegisterCompatibleModes(PlacementMode->GetID(), FBuiltinEditorModes::EM_Geometry);
-		GEditorModeTools().RegisterCompatibleModes(PlacementMode->GetID(), FBuiltinEditorModes::EM_InterpEdit);
-		GEditorModeTools().RegisterCompatibleModes(PlacementMode->GetID(), FBuiltinEditorModes::EM_MeshPaint);
-		GEditorModeTools().RegisterCompatibleModes(PlacementMode->GetID(), FBuiltinEditorModes::EM_Landscape);
-		GEditorModeTools().RegisterCompatibleModes(PlacementMode->GetID(), FBuiltinEditorModes::EM_Foliage);
-		GEditorModeTools().RegisterCompatibleModes(PlacementMode->GetID(), FBuiltinEditorModes::EM_Level);
-		GEditorModeTools().RegisterCompatibleModes(PlacementMode->GetID(), FBuiltinEditorModes::EM_Physics);
-		GEditorModeTools().RegisterCompatibleModes(PlacementMode->GetID(), FBuiltinEditorModes::EM_ActorPicker);
+
+		TArray< FString > RecentlyPlacedAsStrings;
+		GConfig->GetArray(TEXT("PlacementMode"), TEXT("RecentlyPlaced"), RecentlyPlacedAsStrings, GEditorUserSettingsIni);
+
+		//FString ActivePaletteName;
+		//GConfig->GetString( TEXT( "PlacementMode" ), TEXT( "ActivePalette" ), ActivePaletteName, GEditorUserSettingsIni );
+
+		for (int Index = 0; Index < RecentlyPlacedAsStrings.Num(); Index++)
+		{
+			RecentlyPlaced.Add(FActorPlacementInfo(RecentlyPlacedAsStrings[Index]));
+		}
+
+
+		FEditorModeRegistry::Get().RegisterMode<FPlacementMode>(
+			FBuiltinEditorModes::EM_Placement,
+			NSLOCTEXT("PlacementMode", "DisplayName", "Place"),
+			FSlateIcon(FEditorStyle::GetStyleSetName(), "LevelEditor.PlacementMode", "LevelEditor.PlacementMode.Small"),
+			true, 0);
+
+		FAssetRegistryModule& AssetRegistryModule = FModuleManager::GetModuleChecked<FAssetRegistryModule>(TEXT("AssetRegistry"));
+		AssetRegistryModule.Get().OnAssetRemoved().AddRaw( this, &FPlacementModeModule::OnAssetRemoved );
+		AssetRegistryModule.Get().OnAssetRenamed().AddRaw( this, &FPlacementModeModule::OnAssetRenamed );
 	}
 
 	/**
@@ -31,52 +42,150 @@ public:
 	 */
 	virtual void ShutdownModule() override
 	{
-		if ( PlacementMode.IsValid() )
+		FEditorModeRegistry::Get().UnregisterMode( FBuiltinEditorModes::EM_Placement );
+
+		FAssetRegistryModule& AssetRegistryModule = FModuleManager::GetModuleChecked<FAssetRegistryModule>(TEXT("AssetRegistry"));
+		AssetRegistryModule.Get().OnAssetRemoved().RemoveAll( this );
+		AssetRegistryModule.Get().OnAssetRenamed().RemoveAll( this );
+	}
+
+	DECLARE_DERIVED_EVENT( FPlacementModeModule, IPlacementModeModule::FOnRecentlyPlacedChanged, FOnRecentlyPlacedChanged );
+	virtual FOnRecentlyPlacedChanged& OnRecentlyPlacedChanged() override { return RecentlyPlacedChanged; }
+
+	/**
+	 * Add the specified assets to the recently placed items list
+	 */
+	virtual void AddToRecentlyPlaced( const TArray< UObject* >& PlacedObjects, UActorFactory* FactoryUsed = NULL ) override
+	{
+		FString FactoryPath;
+		if (FactoryUsed != NULL)
+		{
+			FactoryPath = FactoryUsed->GetPathName();
+		}
+
+		TArray< UObject* > FilteredPlacedObjects;
+		for (UObject* PlacedObject : PlacedObjects)
+		{
+			// Don't include null placed objects that just have factories.
+			if (PlacedObject == NULL)
+			{
+				continue;
+			}
+
+			// Don't add brush builders to the recently placed.
+			if (PlacedObject->IsA(UBrushBuilder::StaticClass()))
+			{
+				continue;
+			}
+
+			FilteredPlacedObjects.Add(PlacedObject);
+		}
+
+		// Don't change the recently placed if nothing passed the filter.
+		if (FilteredPlacedObjects.Num() == 0)
+		{
+			return;
+		}
+
+		bool Changed = false;
+		for (int Index = 0; Index < FilteredPlacedObjects.Num(); Index++)
 		{
-			GEditorModeTools().UnregisterCompatibleModes(PlacementMode->GetID(), FBuiltinEditorModes::EM_Bsp);
-			GEditorModeTools().UnregisterCompatibleModes(PlacementMode->GetID(), FBuiltinEditorModes::EM_Geometry);
-			GEditorModeTools().UnregisterCompatibleModes(PlacementMode->GetID(), FBuiltinEditorModes::EM_InterpEdit);
-			GEditorModeTools().UnregisterCompatibleModes(PlacementMode->GetID(), FBuiltinEditorModes::EM_MeshPaint);
-			GEditorModeTools().UnregisterCompatibleModes(PlacementMode->GetID(), FBuiltinEditorModes::EM_Landscape);
-			GEditorModeTools().UnregisterCompatibleModes(PlacementMode->GetID(), FBuiltinEditorModes::EM_Foliage);
-			GEditorModeTools().UnregisterCompatibleModes(PlacementMode->GetID(), FBuiltinEditorModes::EM_Level);
-			GEditorModeTools().UnregisterCompatibleModes(PlacementMode->GetID(), FBuiltinEditorModes::EM_Physics);
-			GEditorModeTools().UnregisterCompatibleModes(PlacementMode->GetID(), FBuiltinEditorModes::EM_ActorPicker);
-			GEditorModeTools().UnregisterMode( PlacementMode.ToSharedRef() );
-			PlacementMode = NULL;
+			Changed |= RecentlyPlaced.Remove(FActorPlacementInfo(FilteredPlacedObjects[Index]->GetPathName(), FactoryPath)) > 0;
 		}
+
+		for (int Index = 0; Index < FilteredPlacedObjects.Num(); Index++)
+		{
+			if (FilteredPlacedObjects[Index] != NULL)
+			{
+				RecentlyPlaced.Insert(FActorPlacementInfo(FilteredPlacedObjects[Index]->GetPathName(), FactoryPath), 0);
+				Changed = true;
+			}
+		}
+
+		for (int Index = RecentlyPlaced.Num() - 1; Index >= 20; Index--)
+		{
+			RecentlyPlaced.RemoveAt(Index);
+			Changed = true;
+		}
+
+		if (Changed)
+		{
+			TArray< FString > RecentlyPlacedAsStrings;
+			for (int Index = 0; Index < RecentlyPlaced.Num(); Index++)
+			{
+				RecentlyPlacedAsStrings.Add(RecentlyPlaced[Index].ToString());
+			}
+
+			GConfig->SetArray(TEXT("PlacementMode"), TEXT("RecentlyPlaced"), RecentlyPlacedAsStrings, GEditorUserSettingsIni);
+			RecentlyPlacedChanged.Broadcast(RecentlyPlaced);
+		}
+	}
+
+	void OnAssetRemoved(const FAssetData& /*InRemovedAssetData*/)
+	{
+		RecentlyPlacedChanged.Broadcast(RecentlyPlaced);
 	}
 
-	virtual void AddToRecentlyPlaced( const TArray< UObject* >& Assets, UActorFactory* FactoryUsed = NULL ) override
+	void OnAssetRenamed(const FAssetData& AssetData, const FString& OldObjectPath)
 	{
-		PlacementMode->AddToRecentlyPlaced( Assets, FactoryUsed );
+		for (auto& RecentlyPlacedItem : RecentlyPlaced)
+		{
+			if (RecentlyPlacedItem.ObjectPath == OldObjectPath)
+			{
+				RecentlyPlacedItem.ObjectPath = AssetData.ObjectPath.ToString();
+				break;
+			}
+		}
+
+		RecentlyPlacedChanged.Broadcast(RecentlyPlaced);
 	}
 
+	/**
+	 * Add the specified asset to the recently placed items list
+	 */
 	virtual void AddToRecentlyPlaced( UObject* Asset, UActorFactory* FactoryUsed = NULL ) override
 	{
 		TArray< UObject* > Assets;
 		Assets.Add( Asset );
-		PlacementMode->AddToRecentlyPlaced( Assets, FactoryUsed );
+		AddToRecentlyPlaced( Assets, FactoryUsed );
 	}
 
+	/**
+	 * Get a copy of the recently placed items list
+	 */
 	virtual const TArray< FActorPlacementInfo >& GetRecentlyPlaced() const override
 	{
-		return PlacementMode->GetRecentlyPlaced();
+		return RecentlyPlaced;
 	}
 
-	virtual bool IsPlacementModeAvailable() const override
+	/** @return the event that is broadcast whenever the placement mode enters a placing session */
+	DECLARE_DERIVED_EVENT( FPlacementModeModule, IPlacementModeModule::FOnStartedPlacingEvent, FOnStartedPlacingEvent );
+	virtual FOnStartedPlacingEvent& OnStartedPlacing() override
 	{
-		return PlacementMode.IsValid();
+		return StartedPlacingEvent;
 	}
-
-	virtual TSharedRef< IPlacementMode > GetPlacementMode() const override
+	virtual void BroadcastStartedPlacing( const TArray< UObject* >& Assets ) override
 	{
-		return PlacementMode.ToSharedRef();
+		StartedPlacingEvent.Broadcast( Assets );
 	}
 
+	/** @return the event that is broadcast whenever the placement mode exits a placing session */
+	DECLARE_DERIVED_EVENT( FPlacementModeModule, IPlacementModeModule::FOnStoppedPlacingEvent, FOnStoppedPlacingEvent );
+	virtual FOnStoppedPlacingEvent& OnStoppedPlacing() override
+	{
+		return StoppedPlacingEvent;
+	}
+	virtual void BroadcastStoppedPlacing( bool bWasSuccessfullyPlaced ) override
+	{
+		StoppedPlacingEvent.Broadcast( bWasSuccessfullyPlaced );
+	}
 private:
 
-	TSharedPtr< FPlacementMode > PlacementMode;
+	TArray< FActorPlacementInfo >	RecentlyPlaced;
+	FOnRecentlyPlacedChanged		RecentlyPlacedChanged;
+
+	FOnStartedPlacingEvent			StartedPlacingEvent;
+	FOnStoppedPlacingEvent			StoppedPlacingEvent;
 
 	TArray< TSharedPtr<FExtender> > ContentPaletteFiltersExtenders;
 	TArray< TSharedPtr<FExtender> > PaletteExtenders;
diff --git a/Engine/Source/Editor/PlacementMode/Private/PlacementModeToolkit.h b/Engine/Source/Editor/PlacementMode/Private/PlacementModeToolkit.h
--- a/Engine/Source/Editor/PlacementMode/Private/PlacementModeToolkit.h
+++ b/Engine/Source/Editor/PlacementMode/Private/PlacementModeToolkit.h
@@ -16,7 +16,7 @@ public:
 	/** IToolkit interface */
 	virtual FName GetToolkitFName() const override { return FName("PlacementMode"); }
 	virtual FText GetBaseToolkitName() const override { return NSLOCTEXT("BuilderModeToolkit", "DisplayName", "Builder"); }
-	virtual class FEdMode* GetEditorMode() const override { return GEditorModeTools().GetActiveMode( FBuiltinEditorModes::EM_Default ); }
+	virtual class FEdMode* GetEditorMode() const override { return GLevelEditorModeTools().GetActiveMode( FBuiltinEditorModes::EM_Placement ); }
 	virtual TSharedPtr<class SWidget> GetInlineContent() const override { return PlacementModeTools; }
 
 private:
diff --git a/Engine/Source/Editor/PlacementMode/Private/SPlacementModeTools.cpp b/Engine/Source/Editor/PlacementMode/Private/SPlacementModeTools.cpp
--- a/Engine/Source/Editor/PlacementMode/Private/SPlacementModeTools.cpp
+++ b/Engine/Source/Editor/PlacementMode/Private/SPlacementModeTools.cpp
@@ -227,9 +227,9 @@ const FSlateBrush* SPlacementAssetEntry::GetBorder() const
 
 SPlacementModeTools::~SPlacementModeTools()
 {
-	if ( IPlacementModeModule::IsAvailable() && IPlacementModeModule::Get().IsPlacementModeAvailable() )
+	if ( IPlacementModeModule::IsAvailable() )
 	{
-		IPlacementModeModule::Get().GetPlacementMode()->OnRecentlyPlacedChanged().RemoveAll( this );
+		IPlacementModeModule::Get().OnRecentlyPlacedChanged().RemoveAll( this );
 	}
 }
 
@@ -239,7 +239,7 @@ void SPlacementModeTools::Construct( const FArguments& InArgs )
 	bPlaceablesFullRefreshRequested = false;
 	bVolumesRefreshRequested = false;
 
-	FPlacementMode* PlacementEditMode = (FPlacementMode*)GEditorModeTools().GetActiveMode( FBuiltinEditorModes::EM_Placement );
+	FPlacementMode* PlacementEditMode = (FPlacementMode*)GLevelEditorModeTools().GetActiveMode( FBuiltinEditorModes::EM_Placement );
 	PlacementEditMode->AddValidFocusTargetForPlacement( SharedThis( this ) );
 
 	ChildSlot
@@ -289,7 +289,7 @@ void SPlacementModeTools::Construct( const FArguments& InArgs )
 
 	RefreshRecentlyPlaced();
 
-	IPlacementModeModule::Get().GetPlacementMode()->OnRecentlyPlacedChanged().AddSP( this, &SPlacementModeTools::UpdateRecentlyPlacedAssets );
+	IPlacementModeModule::Get().OnRecentlyPlacedChanged().AddSP( this, &SPlacementModeTools::UpdateRecentlyPlacedAssets );
 }
 
 TSharedRef< SWidget > SPlacementModeTools::CreateStandardPanel()
@@ -913,7 +913,7 @@ FReply SPlacementModeTools::OnKeyDown( const FGeometry& MyGeometry, const FKeybo
 
 	if ( InKeyboardEvent.GetKey() == EKeys::Escape )
 	{
-		FPlacementMode* PlacementEditMode = (FPlacementMode*)GEditorModeTools().GetActiveMode( FBuiltinEditorModes::EM_Placement );
+		FPlacementMode* PlacementEditMode = (FPlacementMode*)GLevelEditorModeTools().GetActiveMode( FBuiltinEditorModes::EM_Placement );
 		PlacementEditMode->StopPlacing();
 		Reply = FReply::Handled();
 	}
diff --git a/Engine/Source/Editor/PlacementMode/Public/IPlacementMode.h b/Engine/Source/Editor/PlacementMode/Public/IPlacementMode.h
--- a/Engine/Source/Editor/PlacementMode/Public/IPlacementMode.h
+++ b/Engine/Source/Editor/PlacementMode/Public/IPlacementMode.h
@@ -4,13 +4,6 @@
 class IPlacementMode
 {
 public:
-	/** @return the event that is broadcast whenever the placement mode exits a placing session */
-	DECLARE_EVENT_OneParam( IPlacementMode, FOnStoppedPlacingEvent, bool /*bWasSuccessfullyPlaced*/ );
-	virtual FOnStoppedPlacingEvent& OnStoppedPlacing() = 0;
-
-	/** @return the event that is broadcast whenever the placement mode enters a placing session */
-	DECLARE_EVENT_OneParam( IPlacementMode, FOnStartedPlacingEvent, const TArray< UObject* >& /*Assets*/ );
-	virtual FOnStartedPlacingEvent& OnStartedPlacing() = 0;
 
 	/** Ends the current placing session in failure, if one is active. Does nothing otherwise. */
 	virtual void StopPlacing() = 0;
@@ -36,16 +29,6 @@ public:
 	/** Removes a widget which when focused would have not ended the active placing session. */
 	virtual void RemoveValidFocusTargetForPlacement( const TWeakPtr< SWidget >& Widget ) = 0;
 
-	/** @return the event that is broadcasted whenever the recently placed actors change  */
-	DECLARE_EVENT_OneParam( IPlacementMode, FOnRecentlyPlacedChanged, const TArray< FActorPlacementInfo >& /*NewRecentlyPlaced*/ );
-	virtual FOnRecentlyPlacedChanged& OnRecentlyPlacedChanged() = 0;
-
-	/** Adds a set of actors and an optional factory to the recently placed list*/
-	virtual void AddToRecentlyPlaced( const TArray< UObject* >& PlacedObjects, UActorFactory* FactoryUsed = NULL ) = 0;
-
-	/** @return an array of actor placement information of the recent placements */
-	virtual const TArray< FActorPlacementInfo >& GetRecentlyPlaced() const = 0;
-
 	/** @return an array of the objects curently being placed */
 	virtual const TArray< TWeakObjectPtr<UObject> >& GetCurrentlyPlacingObjects() const = 0;
 };
diff --git a/Engine/Source/Editor/PlacementMode/Public/IPlacementModeModule.h b/Engine/Source/Editor/PlacementMode/Public/IPlacementModeModule.h
--- a/Engine/Source/Editor/PlacementMode/Public/IPlacementModeModule.h
+++ b/Engine/Source/Editor/PlacementMode/Public/IPlacementModeModule.h
@@ -30,14 +30,39 @@ public:
 		return FModuleManager::Get().IsModuleLoaded( "PlacementMode" );
 	}
 
+	/**
+	 * Add the specified assets to the recently placed items list
+	 */
 	virtual void AddToRecentlyPlaced( const TArray< UObject* >& Assets, UActorFactory* FactoryUsed = NULL ) = 0;
-
+	
+	/**
+	 * Add the specified asset to the recently placed items list
+	 */
 	virtual void AddToRecentlyPlaced( UObject* Asset, UActorFactory* FactoryUsed = NULL ) = 0;
 
+	/**
+	 * Get a copy of the recently placed items
+	 */
 	virtual const TArray< FActorPlacementInfo >& GetRecentlyPlaced() const = 0;
 
-	virtual bool IsPlacementModeAvailable() const = 0;
+	/**
+	 * @return the event that is broadcast whenever the list of recently placed assets changes
+	 */
+	DECLARE_EVENT_OneParam( IPlacementMode, FOnRecentlyPlacedChanged, const TArray< FActorPlacementInfo >& /*NewRecentlyPlaced*/ );
+	virtual FOnRecentlyPlacedChanged& OnRecentlyPlacedChanged() = 0;
 
-	virtual TSharedRef< class IPlacementMode > GetPlacementMode() const = 0;
+	/**
+	 * @return the event that is broadcast whenever a placement mode enters a placing session
+	 */
+	DECLARE_EVENT_OneParam( IPlacementMode, FOnStartedPlacingEvent, const TArray< UObject* >& /*Assets*/ );
+	virtual FOnStartedPlacingEvent& OnStartedPlacing() = 0;
+	virtual void BroadcastStartedPlacing( const TArray< UObject* >& Assets ) = 0;
+
+	/**
+	 * @return the event that is broadcast whenever a placement mode exits a placing session
+	 */
+	DECLARE_EVENT_OneParam( IPlacementMode, FOnStoppedPlacingEvent, bool /*bWasSuccessfullyPlaced*/ );
+	virtual FOnStoppedPlacingEvent& OnStoppedPlacing() = 0;
+	virtual void BroadcastStoppedPlacing( bool bWasSuccessfullyPlaced ) = 0;
 };
 
diff --git a/Engine/Source/Editor/SceneOutliner/Private/SSceneOutliner.cpp b/Engine/Source/Editor/SceneOutliner/Private/SSceneOutliner.cpp
--- a/Engine/Source/Editor/SceneOutliner/Private/SSceneOutliner.cpp
+++ b/Engine/Source/Editor/SceneOutliner/Private/SSceneOutliner.cpp
@@ -3038,7 +3038,7 @@ namespace SceneOutliner
 					// Code from FLevelEditorActionCallbacks::Delete_CanExecute()
 					// Should this be just return FReply::Unhandled()?
 					TArray<FEdMode*> ActiveModes; 
-					GEditorModeTools().GetActiveModes( ActiveModes );
+					GLevelEditorModeTools().GetActiveModes( ActiveModes );
 					for( int32 ModeIndex = 0; ModeIndex < ActiveModes.Num(); ++ModeIndex )
 					{
 						const EEditAction::Type CanProcess = ActiveModes[ModeIndex]->GetActionEditDelete();
