Fixed: CRASH:EDITOR: Editor crashes when unloading a landscape from sublevel in world composition browser
TTP# 343110

diff --git a/Engine/Source/Editor/PlacementMode/Private/PlacementMode.cpp b/Engine/Source/Editor/PlacementMode/Private/PlacementMode.cpp
--- a/Engine/Source/Editor/PlacementMode/Private/PlacementMode.cpp
+++ b/Engine/Source/Editor/PlacementMode/Private/PlacementMode.cpp
@@ -449,7 +449,6 @@ bool FPlacementMode::IsCompatibleWith(FEditorModeID OtherModeID) const
 		OtherModeID == FBuiltinEditorModes::EM_Geometry		||
 		OtherModeID == FBuiltinEditorModes::EM_InterpEdit	||
 		OtherModeID == FBuiltinEditorModes::EM_MeshPaint	||
-		OtherModeID == FBuiltinEditorModes::EM_Landscape	||
 		OtherModeID == FBuiltinEditorModes::EM_Foliage		||
 		OtherModeID == FBuiltinEditorModes::EM_Level		||
 		OtherModeID == FBuiltinEditorModes::EM_Physics		||
