Put doc link widget on Place mode entries

diff --git a/Engine/Source/Editor/PlacementMode/Private/SPlacementModeTools.cpp b/Engine/Source/Editor/PlacementMode/Private/SPlacementModeTools.cpp
--- a/Engine/Source/Editor/PlacementMode/Private/SPlacementModeTools.cpp
+++ b/Engine/Source/Editor/PlacementMode/Private/SPlacementModeTools.cpp
@@ -104,22 +104,27 @@ void SPlacementAssetEntry::Construct(const FArguments& InArgs, UActorFactory* In
 		AssetDisplayName = FText::FromName( AssetData.AssetName );
 	}
 
-	TSharedPtr<IToolTip> ToolTip;
 	FText ActorTypeDisplayName;
 	AActor* DefaultActor = nullptr;
 	if ( FactoryToUse != nullptr )
 	{
 		DefaultActor = FactoryToUse->GetDefaultActor( AssetData );
 		ActorTypeDisplayName = FactoryToUse->GetDisplayName();
-		ToolTip = FEditorClassUtils::GetTooltip(DefaultActor->GetClass());
 	}
 	else if ( IsClass && CastChecked<UClass>( AssetData.GetAsset() )->IsChildOf( AActor::StaticClass() ) )
 	{
 		DefaultActor = CastChecked<AActor>( CastChecked<UClass>( AssetData.GetAsset() )->ClassDefaultObject );
 		ActorTypeDisplayName = FText::FromString( FName::NameToDisplayString( DefaultActor->GetClass()->GetName(), false ) );
-		ToolTip = FEditorClassUtils::GetTooltip(DefaultActor->GetClass());
 	}
 	
+	UClass* DocClass = nullptr;
+	TSharedPtr<IToolTip> ToolTip;
+	if(DefaultActor != nullptr)
+	{
+		DocClass = DefaultActor->GetClass();
+		ToolTip = FEditorClassUtils::GetTooltip(DefaultActor->GetClass());
+	}
+
 	if ( InArgs._LabelOverride.IsEmpty() )
 	{
 		if (IsClass && !IsVolume && !ActorTypeDisplayName.IsEmpty())
@@ -143,6 +148,14 @@ void SPlacementAssetEntry::Construct(const FArguments& InArgs, UActorFactory* In
 	HoverImage = &ButtonStyle.Hovered;
 	PressedImage = &ButtonStyle.Pressed; 
 
+	// Create doc link widget if there is a class to link to
+	TSharedRef<SWidget> DocWidget = SNew(SSpacer);
+	if(DocClass != NULL)
+	{
+		DocWidget = FEditorClassUtils::GetDocumentationLinkWidget(DocClass);
+		DocWidget->SetCursor( EMouseCursor::Default );
+	}
+
 	ChildSlot
 	[
 		SNew( SBorder )
@@ -185,6 +198,13 @@ void SPlacementAssetEntry::Construct(const FArguments& InArgs, UActorFactory* In
 					.HighlightText(InArgs._HighlightText)
 				]
 			]
+
+			+ SHorizontalBox::Slot()
+			.VAlign(VAlign_Center)
+			.AutoWidth()
+			[
+				DocWidget
+			]
 		]
 	];
 }
