Renamed geometry back to BSP

diff --git a/Engine/Source/Editor/BspMode/Private/BspModeModule.cpp b/Engine/Source/Editor/BspMode/Private/BspModeModule.cpp
--- a/Engine/Source/Editor/BspMode/Private/BspModeModule.cpp
+++ b/Engine/Source/Editor/BspMode/Private/BspModeModule.cpp
@@ -13,7 +13,7 @@ void FBspModeModule::StartupModule()
 
 	FEditorModeRegistry::Get().RegisterMode<FBspMode>(
 		FBuiltinEditorModes::EM_Bsp,
-		NSLOCTEXT("GeometryMode", "DisplayName", "Geometry"),
+		NSLOCTEXT("GeometryMode", "DisplayName", "Geometry Editing"),
 		FSlateIcon(FEditorStyle::GetStyleSetName(), "LevelEditor.BspMode", "LevelEditor.BspMode.Small"),
 		false,		// Visible
 		100			// UI priority order
diff --git a/Engine/Source/Editor/PlacementMode/Private/SPlacementModeTools.cpp b/Engine/Source/Editor/PlacementMode/Private/SPlacementModeTools.cpp
--- a/Engine/Source/Editor/PlacementMode/Private/SPlacementModeTools.cpp
+++ b/Engine/Source/Editor/PlacementMode/Private/SPlacementModeTools.cpp
@@ -19,10 +19,10 @@ namespace EPlacementTab
 	enum Type
 	{
 		RecentlyPlaced = 0,
-		Geometry,
+		Basic,
 		Lights,
 		Visual,
-		Basic,
+		Geometry,
 		Volumes,
 		AllClasses,
 	};
@@ -347,13 +347,14 @@ TSharedRef< SWidget > SPlacementModeTools::CreateStandardPanel()
 			]
 		]
 
+
 		+ SVerticalBox::Slot()
 		.AutoHeight()
 		[
-			SNew( SBox )
-			.AddMetaData<FTagMetaData>(FTagMetaData(TEXT("PMGeometry")))
+			SNew(SBox)
+			.AddMetaData<FTagMetaData>(FTagMetaData(TEXT("PMBasic")))
 			[
-				CreatePlacementGroupTab( (int32)EPlacementTab::Geometry, NSLOCTEXT( "PlacementMode", "Geometry", "Geometry" ), false )
+				CreatePlacementGroupTab((int32)EPlacementTab::Basic, NSLOCTEXT("PlacementMode", "Basic", "Basic"), false)
 			]
 		]
 
@@ -377,13 +378,14 @@ TSharedRef< SWidget > SPlacementModeTools::CreateStandardPanel()
 			]
 		]
 
+
 		+ SVerticalBox::Slot()
 		.AutoHeight()
 		[
-			SNew( SBox )
-			.AddMetaData<FTagMetaData>(FTagMetaData(TEXT("PMBasic")))
+			SNew(SBox)
+			.AddMetaData<FTagMetaData>(FTagMetaData(TEXT("PMGeometry")))
 			[
-				CreatePlacementGroupTab( (int32)EPlacementTab::Basic, NSLOCTEXT( "PlacementMode", "Basic", "Basic" ), false )
+				CreatePlacementGroupTab((int32)EPlacementTab::Geometry, NSLOCTEXT("PlacementMode", "BSP", "BSP"), false)
 			]
 		]
 
@@ -423,7 +425,7 @@ TSharedRef< SWidget > SPlacementModeTools::CreateStandardPanel()
 			.BorderImage( FEditorStyle::GetBrush( "ToolPanel.DarkGroupBorder" ) )
 			[
 				SAssignNew( WidgetSwitcher, SWidgetSwitcher )
-				.WidgetIndex( (int32)EPlacementTab::Geometry )
+				.WidgetIndex( (int32)EPlacementTab::Basic )
 
 				// Recently Placed
 				+ SWidgetSwitcher::Slot()
@@ -436,10 +438,15 @@ TSharedRef< SWidget > SPlacementModeTools::CreateStandardPanel()
 					]
 				]
 
-				// Geometry
+				// Basics
 				+ SWidgetSwitcher::Slot()
 				[
-					FModuleManager::LoadModuleChecked<IBspModeModule>("BspMode").CreateBspModeWidget()
+					SNew(SScrollBox)
+
+					+ SScrollBox::Slot()
+					[
+						BuildBasicWidget()
+					]
 				]
 
 				// Lights
@@ -464,15 +471,10 @@ TSharedRef< SWidget > SPlacementModeTools::CreateStandardPanel()
 					]
 				]
 
-				// Basics
-				+ SWidgetSwitcher::Slot()
+				// Geometry
+				+SWidgetSwitcher::Slot()
 				[
-					SNew( SScrollBox )
-
-					+ SScrollBox::Slot()
-					[
-						BuildBasicWidget()
-					]
+					FModuleManager::LoadModuleChecked<IBspModeModule>("BspMode").CreateBspModeWidget()
 				]
 
 				// Volumes
@@ -674,32 +676,33 @@ TSharedRef< SWidget > SPlacementModeTools::BuildBasicWidget()
 		BuildDraggableAssetWidget( UActorFactoryPlayerStart::StaticClass() )
 	]
 
-	// Triggers
 	+ SVerticalBox::Slot()
 	.AutoHeight()
 	[
-		BuildDraggableAssetWidget( UActorFactoryTriggerBox::StaticClass() )
+		BuildDraggableAssetWidget(UActorFactoryPointLight::StaticClass())
 	]
 
+
+	// Triggers
 	+ SVerticalBox::Slot()
 	.AutoHeight()
 	[
-		BuildDraggableAssetWidget( UActorFactoryTriggerSphere::StaticClass() )
+		BuildDraggableAssetWidget( UActorFactoryTriggerBox::StaticClass() )
 	]
 
 	+ SVerticalBox::Slot()
 	.AutoHeight()
 	[
-		BuildDraggableAssetWidget( UActorFactoryTriggerCapsule::StaticClass() )
+		BuildDraggableAssetWidget( UActorFactoryTriggerSphere::StaticClass() )
 	]
 
-	// Misc
 	+ SVerticalBox::Slot()
 	.AutoHeight()
 	[
-		BuildDraggableAssetWidget( UActorFactoryNote::StaticClass() )
+		BuildDraggableAssetWidget( UActorFactoryTriggerCapsule::StaticClass() )
 	]
 
+
 	+ SVerticalBox::Slot()
 	.AutoHeight()
 	[
