Prevented BP-derived volumes from being populated in the volumes list in placement mode.

#change Prevented BP-derived volumes from being populated in the Volumes list in placement mode, and from being included by UUnrealEdEngine::GetSortedVolumeClasses as plenty of code expects this to be a static list.

diff --git a/Engine/Source/Editor/PlacementMode/Private/SPlacementModeTools.cpp b/Engine/Source/Editor/PlacementMode/Private/SPlacementModeTools.cpp
--- a/Engine/Source/Editor/PlacementMode/Private/SPlacementModeTools.cpp
+++ b/Engine/Source/Editor/PlacementMode/Private/SPlacementModeTools.cpp
@@ -775,15 +775,15 @@ void SPlacementModeTools::RefreshVolumes()
 	FAssetData NoAssetData;
 	for ( TObjectIterator<UClass> ClassIt; ClassIt; ++ClassIt )
 	{
-		// Don't offer skeleton classes
-		bool bIsSkeletonClass = FKismetEditorUtilities::IsClassABlueprintSkeleton(*ClassIt);
+		const UClass* Class = *ClassIt;
 
-		if ( !ClassIt->HasAllClassFlags( CLASS_NotPlaceable ) &&
-			 !ClassIt->HasAnyClassFlags( CLASS_Abstract | CLASS_Deprecated | CLASS_NewerVersionExists ) &&
-			  ClassIt->IsChildOf( AVolume::StaticClass() ) )
+		if (!Class->HasAllClassFlags(CLASS_NotPlaceable) &&
+			!Class->HasAnyClassFlags(CLASS_Abstract | CLASS_Deprecated | CLASS_NewerVersionExists) &&
+			Class->IsChildOf(AVolume::StaticClass()) &&
+			Class->ClassGeneratedBy == nullptr )
 		{
-			UActorFactory* Factory = GEditor->FindActorFactoryByClassForActorClass( UActorFactoryBoxVolume::StaticClass(), *ClassIt );
-			Entries.Add( SNew( SPlacementAssetEntry, Factory, FAssetData( *ClassIt ) ) );
+			UActorFactory* Factory = GEditor->FindActorFactoryByClassForActorClass(UActorFactoryBoxVolume::StaticClass(), Class);
+			Entries.Add(SNew(SPlacementAssetEntry, Factory, FAssetData(Class)));
 		}
 	}
 
