The search box in the placement mode tools panel now filters classes by their localized name as well as their original name in English.

diff --git a/Engine/Source/Editor/PlacementMode/Private/SPlacementModeTools.cpp b/Engine/Source/Editor/PlacementMode/Private/SPlacementModeTools.cpp
--- a/Engine/Source/Editor/PlacementMode/Private/SPlacementModeTools.cpp
+++ b/Engine/Source/Editor/PlacementMode/Private/SPlacementModeTools.cpp
@@ -917,7 +917,8 @@ void SPlacementModeTools::RefreshPlaceables()
 		// Filter out the widgets that don't belong
 		for ( int32 i = 0; i < PlaceableClassWidgets.Num(); i++ )
 		{
-			if ( PlaceableClassWidgets[i]->AssetDisplayName.ToString().Contains(SearchText.ToString()) )
+			const FString* SourceString = FTextInspector::GetSourceString(PlaceableClassWidgets[i]->AssetDisplayName);
+			if ( PlaceableClassWidgets[i]->AssetDisplayName.ToString().Contains( SearchText.ToString() ) || ( SourceString && SourceString->Contains( SearchText.ToString() ) ) )
 			{
 				SearchResultsContainer->AddSlot()
 				[
