[INTEGRATE] Change 2427078 by [REDACTED_PERSON][REDACTED_HANDLE] on 2015/01/31 12:24:27

	Updated asset icons for basic shapes

diff --git a/Engine/Source/Editor/PlacementMode/Private/SPlacementModeTools.cpp b/Engine/Source/Editor/PlacementMode/Private/SPlacementModeTools.cpp
--- a/Engine/Source/Editor/PlacementMode/Private/SPlacementModeTools.cpp
+++ b/Engine/Source/Editor/PlacementMode/Private/SPlacementModeTools.cpp
@@ -13,6 +13,7 @@
 #include "SWidgetSwitcher.h"
 #include "GameFramework/Volume.h"
 #include "Engine/PostProcessVolume.h"
+#include "AssetToolsModule.h"
 
 /**
  * These are the tab indexes, if the tabs are reorganized you need to adjust the
@@ -41,15 +42,22 @@ public:
 	SLATE_BEGIN_ARGS( SPlacementAssetThumbnail )
 		: _Width( 32 )
 		, _Height( 32 )
+		, _AlwaysUseGenericThumbnail( false )
+		, _AssetTypeColorOverride()
 	{}
 
 	SLATE_ARGUMENT( uint32, Width )
 
 	SLATE_ARGUMENT( uint32, Height )
 
+	SLATE_ARGUMENT( FName, ClassThumbnailBrushOverride )
+
+	SLATE_ARGUMENT( bool, AlwaysUseGenericThumbnail )
+
+	SLATE_ARGUMENT( TOptional<FLinearColor>, AssetTypeColorOverride )
 	SLATE_END_ARGS()
 
-	void Construct( const FArguments& InArgs, const FAssetData& InAsset, bool bAlwaysUseClassThumbnail )
+	void Construct( const FArguments& InArgs, const FAssetData& InAsset)
 	{
 		Asset = InAsset;
 
@@ -59,7 +67,9 @@ public:
 		Thumbnail = MakeShareable(new FAssetThumbnail(Asset, InArgs._Width, InArgs._Height, ThumbnailPool));
 
 		FAssetThumbnailConfig Config;
-		Config.bForceGenericThumbnail = bAlwaysUseClassThumbnail;
+		Config.bForceGenericThumbnail = InArgs._AlwaysUseGenericThumbnail;
+		Config.ClassThumbnailBrushOverride = InArgs._ClassThumbnailBrushOverride;
+		Config.AssetTypeColorOverride = InArgs._AssetTypeColorOverride;
 		ChildSlot
 		[
 			Thumbnail->MakeThumbnailWidget( Config )
@@ -167,7 +177,10 @@ void SPlacementAssetEntry::Construct(const FArguments& InArgs, UActorFactory* In
 					.WidthOverride( 35 )
 					.HeightOverride( 35 )
 					[
-						SNew( SPlacementAssetThumbnail, AssetData, InArgs._UseClassThumbnailForAsset )
+						SNew( SPlacementAssetThumbnail, AssetData )
+						.ClassThumbnailBrushOverride( InArgs._ClassThumbnailBrushOverride )
+						.AlwaysUseGenericThumbnail( InArgs._AlwaysUseGenericThumbnail )
+						.AssetTypeColorOverride( InArgs._AssetTypeColorOverride )
 					]
 				]
 			]
@@ -569,10 +582,14 @@ TSharedRef< SWidget > BuildDraggableAssetWidget( UClass* InAssetClass )
 	return SNew( SPlacementAssetEntry, Factory, AssetData );
 }
 
-TSharedRef< SWidget > BuildDraggableAssetWidget( UClass* InAssetClass, const FAssetData& InAssetData, bool bUseClassThumbnail )
+
+TSharedRef< SWidget > BuildDraggableAssetWidget( UClass* InAssetClass, const FAssetData& InAssetData, FName InClassThumbnailBrushOverride = NAME_None, TOptional<FLinearColor> InAssetTypeColorOverride = TOptional<FLinearColor>() ) 
 {
 	UActorFactory* Factory = GEditor->FindActorFactoryByClass( InAssetClass );
-	return SNew( SPlacementAssetEntry, Factory, InAssetData ).UseClassThumbnailForAsset( bUseClassThumbnail );
+	return SNew( SPlacementAssetEntry, Factory, InAssetData )
+		.ClassThumbnailBrushOverride( InClassThumbnailBrushOverride )
+		.AlwaysUseGenericThumbnail(true)
+		.AssetTypeColorOverride( InAssetTypeColorOverride );
 }
 
 TSharedRef< SWidget > SPlacementModeTools::BuildLightsWidget()
@@ -651,10 +668,19 @@ TSharedRef< SWidget > SPlacementModeTools::BuildVisualWidget()
 
 TSharedRef< SWidget > SPlacementModeTools::BuildBasicWidget()
 {
-	TSharedRef<SVerticalBox> VerticalBox = SNew( SVerticalBox )
+	// Get color for basic shapes.  It should appear like all the other basic types
+	TOptional<FLinearColor> BasicShapeColorOverride;
+	FAssetToolsModule& AssetToolsModule = FModuleManager::LoadModuleChecked<FAssetToolsModule>(TEXT("AssetTools"));
+	TSharedPtr<IAssetTypeActions> AssetTypeActions;
+	AssetTypeActions = AssetToolsModule.Get().GetAssetTypeActionsForClass(UClass::StaticClass()).Pin();
+	if( AssetTypeActions.IsValid() )
+	{
+		BasicShapeColorOverride = TOptional<FLinearColor>( AssetTypeActions->GetTypeColor() );
+	}
+	
 
 	// Basics
-	
+	TSharedRef<SVerticalBox> VerticalBox = SNew( SVerticalBox )
 	+SVerticalBox::Slot()
 	.AutoHeight()
 	[
@@ -669,25 +695,25 @@ TSharedRef< SWidget > SPlacementModeTools::BuildBasicWidget()
 	.AutoHeight()
 	[
 		// Cube
-		BuildDraggableAssetWidget(UActorFactoryBasicShape::StaticClass(), FAssetData( LoadObject<UStaticMesh>(nullptr,TEXT("/Engine/BasicShapes/Cube.Cube"))), true )
+		BuildDraggableAssetWidget(UActorFactoryBasicShape::StaticClass(), FAssetData( LoadObject<UStaticMesh>(nullptr,TEXT("/Engine/BasicShapes/Cube.Cube"))), FName("ClassThumbnail.Cube"), BasicShapeColorOverride )
 	]
 	+ SVerticalBox::Slot()
 	.AutoHeight()
 	[
 		// Sphere
-		BuildDraggableAssetWidget(UActorFactoryBasicShape::StaticClass(), FAssetData(LoadObject<UStaticMesh>(nullptr,TEXT("/Engine/BasicShapes/Sphere.Sphere"))), true )
+		BuildDraggableAssetWidget(UActorFactoryBasicShape::StaticClass(), FAssetData(LoadObject<UStaticMesh>(nullptr,TEXT("/Engine/BasicShapes/Sphere.Sphere"))), FName("ClassThumbnail.Sphere"), BasicShapeColorOverride )
 	]
 	+ SVerticalBox::Slot()
 	.AutoHeight()
 	[
 		// Cylinder
-		BuildDraggableAssetWidget(UActorFactoryBasicShape::StaticClass(), FAssetData(LoadObject<UStaticMesh>(nullptr,TEXT("/Engine/BasicShapes/Cylinder.Cylinder"))), true )
+		BuildDraggableAssetWidget(UActorFactoryBasicShape::StaticClass(), FAssetData(LoadObject<UStaticMesh>(nullptr,TEXT("/Engine/BasicShapes/Cylinder.Cylinder"))), FName("ClassThumbnail.Cylinder"), BasicShapeColorOverride )
 	]
 	+ SVerticalBox::Slot()
 	.AutoHeight()
 	[
 		// Cone
-		BuildDraggableAssetWidget(UActorFactoryBasicShape::StaticClass(), FAssetData(LoadObject<UStaticMesh>(nullptr,TEXT("/Engine/BasicShapes/Cone.Cone"))), true )
+		BuildDraggableAssetWidget(UActorFactoryBasicShape::StaticClass(), FAssetData(LoadObject<UStaticMesh>(nullptr,TEXT("/Engine/BasicShapes/Cone.Cone"))), FName("ClassThumbnail.Cone"), BasicShapeColorOverride )
 	]
 	+ SVerticalBox::Slot()
 	.AutoHeight()
diff --git a/Engine/Source/Editor/PlacementMode/Private/SPlacementModeTools.h b/Engine/Source/Editor/PlacementMode/Private/SPlacementModeTools.h
--- a/Engine/Source/Editor/PlacementMode/Private/SPlacementModeTools.h
+++ b/Engine/Source/Editor/PlacementMode/Private/SPlacementModeTools.h
@@ -15,7 +15,7 @@ public:
 
 	SLATE_BEGIN_ARGS(SPlacementAssetEntry)
 		: _LabelOverride()
-		, _UseClassThumbnailForAsset(false)
+		, _AlwaysUseGenericThumbnail(false)
 	{}
 	
 	/** Highlight this text in the text block */
@@ -23,7 +23,11 @@ public:
 
 	SLATE_ARGUMENT(FText, LabelOverride)
 
-	SLATE_ARGUMENT(bool, UseClassThumbnailForAsset)
+	SLATE_ARGUMENT(FName, ClassThumbnailBrushOverride)
+
+	SLATE_ARGUMENT(bool, AlwaysUseGenericThumbnail)
+
+	SLATE_ARGUMENT(TOptional<FLinearColor>, AssetTypeColorOverride)
 	SLATE_END_ARGS()
 
 	void Construct(const FArguments& InArgs, UActorFactory* InFactory, const FAssetData& InAsset);
diff --git a/Engine/Source/Editor/UnrealEd/Private/AssetThumbnail.cpp b/Engine/Source/Editor/UnrealEd/Private/AssetThumbnail.cpp
--- a/Engine/Source/Editor/UnrealEd/Private/AssetThumbnail.cpp
+++ b/Engine/Source/Editor/UnrealEd/Private/AssetThumbnail.cpp
@@ -23,6 +23,7 @@ public:
 		, _HighlightedText(FText::GetEmpty())
 		, _HintColorAndOpacity(FLinearColor(0.0f, 0.0f, 0.0f, 0.0f))
 		, _ClassThumbnailBrushOverride(NAME_None)
+		, _AssetTypeColorOverride()
 		{}
 
 		SLATE_ARGUMENT( FName, Style )
@@ -36,6 +37,7 @@ public:
 		SLATE_ATTRIBUTE( FText, HighlightedText )
 		SLATE_ATTRIBUTE( FLinearColor, HintColorAndOpacity )
 		SLATE_ARGUMENT( FName, ClassThumbnailBrushOverride )
+		SLATE_ARGUMENT( TOptional<FLinearColor>, AssetTypeColorOverride )
 
 	SLATE_END_ARGS()
 
@@ -66,7 +68,11 @@ public:
 		}
 
 		AssetColor = FLinearColor::White;
-		if ( AssetTypeActions.IsValid() )
+		if( InArgs._AssetTypeColorOverride.IsSet() )
+		{
+			AssetColor = InArgs._AssetTypeColorOverride.GetValue();
+		}
+		else if ( AssetTypeActions.IsValid() )
 		{
 			AssetColor = AssetTypeActions->GetTypeColor();
 		}
@@ -773,7 +779,8 @@ TSharedRef<SWidget> FAssetThumbnail::MakeThumbnailWidget( const FAssetThumbnailC
 		.HintColorAndOpacity( InConfig.HintColorAndOpacity )
 		.AllowHintText( InConfig.bAllowHintText )
 		.ClassThumbnailBrushOverride( InConfig.ClassThumbnailBrushOverride )
-		.AllowAssetSpecificThumbnailOverlay( InConfig.bAllowAssetSpecificThumbnailOverlay );
+		.AllowAssetSpecificThumbnailOverlay( InConfig.bAllowAssetSpecificThumbnailOverlay )
+		.AssetTypeColorOverride( InConfig.AssetTypeColorOverride );
 }
 
 void FAssetThumbnail::RefreshThumbnail()
diff --git a/Engine/Source/Editor/UnrealEd/Public/AssetThumbnail.h b/Engine/Source/Editor/UnrealEd/Public/AssetThumbnail.h
--- a/Engine/Source/Editor/UnrealEd/Public/AssetThumbnail.h
+++ b/Engine/Source/Editor/UnrealEd/Public/AssetThumbnail.h
@@ -30,6 +30,7 @@ struct FAssetThumbnailConfig
 		, ThumbnailLabel( EThumbnailLabel::ClassName )
 		, HighlightedText( FText::GetEmpty() )
 		, HintColorAndOpacity( FLinearColor( 0.0f, 0.0f, 0.0f, 0.0f ) )
+		, AssetTypeColorOverride()
 	{
 	}
 
@@ -41,6 +42,7 @@ struct FAssetThumbnailConfig
 	EThumbnailLabel::Type ThumbnailLabel;
 	TAttribute< FText > HighlightedText;
 	TAttribute< FLinearColor > HintColorAndOpacity;
+	TOptional< FLinearColor > AssetTypeColorOverride;
 };
 
 /**
