[INTEGRATE] Change 2436836 by [REDACTED_PERSON][REDACTED_HANDLE] on 2015/02/07 14:06:42

	Fix empty actor types and basic shapes not appearing in the placement browser all classes menu or in search

diff --git a/Engine/Source/Editor/PlacementMode/Private/SPlacementModeTools.cpp b/Engine/Source/Editor/PlacementMode/Private/SPlacementModeTools.cpp
--- a/Engine/Source/Editor/PlacementMode/Private/SPlacementModeTools.cpp
+++ b/Engine/Source/Editor/PlacementMode/Private/SPlacementModeTools.cpp
@@ -575,21 +575,22 @@ const FSlateBrush* SPlacementModeTools::PlacementGroupBorderImage( int32 Placeme
 	}
 }
 
-TSharedRef< SWidget > BuildDraggableAssetWidget( UClass* InAssetClass )
+TSharedRef< SPlacementAssetEntry > BuildDraggableAssetWidget( UClass* InAssetClass, TAttribute<FText> HightlightText = TAttribute<FText>() )
 {
 	UActorFactory* Factory = GEditor->FindActorFactoryByClass( InAssetClass );
 	FAssetData AssetData = FAssetData( Factory->GetDefaultActorClass( FAssetData() ) );
-	return SNew( SPlacementAssetEntry, Factory, AssetData );
+	return SNew( SPlacementAssetEntry, Factory, AssetData ).HighlightText( HightlightText );
 }
 
 
-TSharedRef< SWidget > BuildDraggableAssetWidget( UClass* InAssetClass, const FAssetData& InAssetData, FName InClassThumbnailBrushOverride = NAME_None, TOptional<FLinearColor> InAssetTypeColorOverride = TOptional<FLinearColor>() ) 
+TSharedRef< SPlacementAssetEntry > BuildDraggableAssetWidget( UClass* InAssetClass, const FAssetData& InAssetData, FName InClassThumbnailBrushOverride = NAME_None, TOptional<FLinearColor> InAssetTypeColorOverride = TOptional<FLinearColor>(), TAttribute<FText> HightlightText = TAttribute<FText>() ) 
 {
 	UActorFactory* Factory = GEditor->FindActorFactoryByClass( InAssetClass );
 	return SNew( SPlacementAssetEntry, Factory, InAssetData )
 		.ClassThumbnailBrushOverride( InClassThumbnailBrushOverride )
 		.AlwaysUseGenericThumbnail(true)
-		.AssetTypeColorOverride( InAssetTypeColorOverride );
+		.AssetTypeColorOverride( InAssetTypeColorOverride )
+		.HighlightText( HightlightText );
 }
 
 TSharedRef< SWidget > SPlacementModeTools::BuildLightsWidget()
@@ -666,19 +667,26 @@ TSharedRef< SWidget > SPlacementModeTools::BuildVisualWidget()
 	];
 }
 
-TSharedRef< SWidget > SPlacementModeTools::BuildBasicWidget()
+static TOptional<FLinearColor> GetBasicShapeColorOverride()
 {
 	// Get color for basic shapes.  It should appear like all the other basic types
-	TOptional<FLinearColor> BasicShapeColorOverride;
-	FAssetToolsModule& AssetToolsModule = FModuleManager::LoadModuleChecked<FAssetToolsModule>(TEXT("AssetTools"));
-	TSharedPtr<IAssetTypeActions> AssetTypeActions;
-	AssetTypeActions = AssetToolsModule.Get().GetAssetTypeActionsForClass(UClass::StaticClass()).Pin();
-	if( AssetTypeActions.IsValid() )
+	static TOptional<FLinearColor> BasicShapeColorOverride;
+	
+	if( !BasicShapeColorOverride.IsSet() )
 	{
-		BasicShapeColorOverride = TOptional<FLinearColor>( AssetTypeActions->GetTypeColor() );
+		FAssetToolsModule& AssetToolsModule = FModuleManager::LoadModuleChecked<FAssetToolsModule>(TEXT("AssetTools"));
+		TSharedPtr<IAssetTypeActions> AssetTypeActions;
+		AssetTypeActions = AssetToolsModule.Get().GetAssetTypeActionsForClass(UClass::StaticClass()).Pin();
+		if (AssetTypeActions.IsValid())
+		{
+			BasicShapeColorOverride = TOptional<FLinearColor>(AssetTypeActions->GetTypeColor());
+		}
 	}
-	
+	return BasicShapeColorOverride;
+}
 
+TSharedRef< SWidget > SPlacementModeTools::BuildBasicWidget()
+{
 	// Basics
 	TSharedRef<SVerticalBox> VerticalBox = SNew( SVerticalBox )
 	+SVerticalBox::Slot()
@@ -695,25 +703,25 @@ TSharedRef< SWidget > SPlacementModeTools::BuildBasicWidget()
 	.AutoHeight()
 	[
 		// Cube
-		BuildDraggableAssetWidget(UActorFactoryBasicShape::StaticClass(), FAssetData( LoadObject<UStaticMesh>(nullptr,*UActorFactoryBasicShape::BasicCube.ToString())), FName("ClassThumbnail.Cube"), BasicShapeColorOverride )
+		BuildDraggableAssetWidget(UActorFactoryBasicShape::StaticClass(), FAssetData( LoadObject<UStaticMesh>(nullptr,*UActorFactoryBasicShape::BasicCube.ToString())), FName("ClassThumbnail.Cube"), GetBasicShapeColorOverride() )
 	]
 	+ SVerticalBox::Slot()
 	.AutoHeight()
 	[
 		// Sphere
-		BuildDraggableAssetWidget(UActorFactoryBasicShape::StaticClass(), FAssetData(LoadObject<UStaticMesh>(nullptr,*UActorFactoryBasicShape::BasicSphere.ToString())), FName("ClassThumbnail.Sphere"), BasicShapeColorOverride )
+		BuildDraggableAssetWidget(UActorFactoryBasicShape::StaticClass(), FAssetData(LoadObject<UStaticMesh>(nullptr,*UActorFactoryBasicShape::BasicSphere.ToString())), FName("ClassThumbnail.Sphere"), GetBasicShapeColorOverride() )
 	]
 	+ SVerticalBox::Slot()
 	.AutoHeight()
 	[
 		// Cylinder
-		BuildDraggableAssetWidget(UActorFactoryBasicShape::StaticClass(), FAssetData(LoadObject<UStaticMesh>(nullptr,*UActorFactoryBasicShape::BasicCylinder.ToString())), FName("ClassThumbnail.Cylinder"), BasicShapeColorOverride )
+		BuildDraggableAssetWidget(UActorFactoryBasicShape::StaticClass(), FAssetData(LoadObject<UStaticMesh>(nullptr,*UActorFactoryBasicShape::BasicCylinder.ToString())), FName("ClassThumbnail.Cylinder"), GetBasicShapeColorOverride() )
 	]
 	+ SVerticalBox::Slot()
 	.AutoHeight()
 	[
 		// Cone
-		BuildDraggableAssetWidget(UActorFactoryBasicShape::StaticClass(), FAssetData(LoadObject<UStaticMesh>(nullptr,*UActorFactoryBasicShape::BasicCone.ToString())), FName("ClassThumbnail.Cone"), BasicShapeColorOverride )
+		BuildDraggableAssetWidget(UActorFactoryBasicShape::StaticClass(), FAssetData(LoadObject<UStaticMesh>(nullptr,*UActorFactoryBasicShape::BasicCone.ToString())), FName("ClassThumbnail.Cone"), GetBasicShapeColorOverride() )
 	]
 	+ SVerticalBox::Slot()
 	.AutoHeight()
@@ -858,6 +866,8 @@ void SPlacementModeTools::RebuildPlaceableClassWidgetCache()
 
 	TArray< TSharedRef<SPlacementAssetEntry> > Entries;
 
+	TAttribute<FText> HighlightTextAttrib( this, &SPlacementModeTools::GetHighlightText );
+
 	// Add loaded classes
 	FText UnusedErrorMessage;
 	FAssetData NoAssetData;
@@ -882,7 +892,7 @@ void SPlacementModeTools::RebuildPlaceableClassWidgetCache()
 			{
 				UActorFactory* Factory = GEditor->FindActorFactoryByClassForActorClass(UActorFactoryBoxVolume::StaticClass(), *ClassIt);
 				Entry = SNew(SPlacementAssetEntry, Factory, FAssetData(*ClassIt))
-					.HighlightText(this, &SPlacementModeTools::GetHighlightText);
+					.HighlightText(HighlightTextAttrib);
 			}
 			else
 			{
@@ -891,12 +901,12 @@ void SPlacementModeTools::RebuildPlaceableClassWidgetCache()
 					if ( !ActorFactory )
 					{
 						Entry = SNew(SPlacementAssetEntry, nullptr, FAssetData(*ClassIt))
-							.HighlightText(this, &SPlacementModeTools::GetHighlightText);
+							.HighlightText(HighlightTextAttrib);
 					}
 					else
 					{
 						Entry = SNew(SPlacementAssetEntry, *ActorFactory, FAssetData(*ClassIt))
-							.HighlightText(this, &SPlacementModeTools::GetHighlightText);
+							.HighlightText(HighlightTextAttrib);
 					}
 				}
 			}
@@ -908,6 +918,16 @@ void SPlacementModeTools::RebuildPlaceableClassWidgetCache()
 		}
 	}
 
+
+	// Add empty types and shapes which have special actor factory requirements
+	Entries.Add(BuildDraggableAssetWidget(UActorFactoryEmptyActor::StaticClass(),HighlightTextAttrib));
+	Entries.Add(BuildDraggableAssetWidget(UActorFactoryCharacter::StaticClass(),HighlightTextAttrib));
+	Entries.Add(BuildDraggableAssetWidget(UActorFactoryPawn::StaticClass(),HighlightTextAttrib));
+	Entries.Add(BuildDraggableAssetWidget(UActorFactoryBasicShape::StaticClass(), FAssetData(LoadObject<UStaticMesh>(nullptr, *UActorFactoryBasicShape::BasicCube.ToString())), FName("ClassThumbnail.Cube"), GetBasicShapeColorOverride(),HighlightTextAttrib));
+	Entries.Add(BuildDraggableAssetWidget(UActorFactoryBasicShape::StaticClass(), FAssetData(LoadObject<UStaticMesh>(nullptr, *UActorFactoryBasicShape::BasicSphere.ToString())), FName("ClassThumbnail.Sphere"), GetBasicShapeColorOverride(),HighlightTextAttrib));
+	Entries.Add(BuildDraggableAssetWidget(UActorFactoryBasicShape::StaticClass(), FAssetData(LoadObject<UStaticMesh>(nullptr, *UActorFactoryBasicShape::BasicCylinder.ToString())), FName("ClassThumbnail.Cylinder"), GetBasicShapeColorOverride(),HighlightTextAttrib));
+	Entries.Add(BuildDraggableAssetWidget(UActorFactoryBasicShape::StaticClass(), FAssetData(LoadObject<UStaticMesh>(nullptr, *UActorFactoryBasicShape::BasicCone.ToString())), FName("ClassThumbnail.Cone"), GetBasicShapeColorOverride(),HighlightTextAttrib));
+
 	struct FCompareFactoryByDisplayName
 	{
 		FORCEINLINE bool operator()(const TSharedRef<SPlacementAssetEntry>& A, const TSharedRef<SPlacementAssetEntry>& B) const
