Converted the placement browser to use a TTextFilter

diff --git a/Engine/Source/Editor/PlacementMode/Private/SPlacementModeTools.cpp b/Engine/Source/Editor/PlacementMode/Private/SPlacementModeTools.cpp
--- a/Engine/Source/Editor/PlacementMode/Private/SPlacementModeTools.cpp
+++ b/Engine/Source/Editor/PlacementMode/Private/SPlacementModeTools.cpp
@@ -33,6 +33,22 @@ namespace EPlacementTab
 	};
 }
 
+namespace PlacementViewFilter
+{
+
+void GetBasicStrings(TSharedRef<SPlacementAssetEntry> InPlacementAsset, TArray<FString>& OutBasicStrings)
+{
+	OutBasicStrings.Add(InPlacementAsset->AssetDisplayName.ToString());
+
+	const FString* SourceString = FTextInspector::GetSourceString(InPlacementAsset->AssetDisplayName);
+	if (SourceString)
+	{
+		OutBasicStrings.Add(*SourceString);
+	}
+}
+
+} // namespace PlacementViewFilter
+
 /**
  * These are the asset thumbnails.
  */
@@ -278,6 +294,10 @@ void SPlacementModeTools::Construct( const FArguments& InArgs )
 	FPlacementMode* PlacementEditMode = (FPlacementMode*)GLevelEditorModeTools().GetActiveMode( FBuiltinEditorModes::EM_Placement );
 	PlacementEditMode->AddValidFocusTargetForPlacement( SharedThis( this ) );
 
+	SearchTextFilter = MakeShareable(new FPlacementAssetEntryTextFilter(
+		FPlacementAssetEntryTextFilter::FItemToStringArray::CreateStatic(&PlacementViewFilter::GetBasicStrings)
+		));
+
 	ChildSlot
 	[
 		SNew( SVerticalBox )
@@ -286,7 +306,7 @@ void SPlacementModeTools::Construct( const FArguments& InArgs )
 		.Padding(4)
 		.AutoHeight()
 		[
-			SNew( SSearchBox )
+			SAssignNew( SearchBoxPtr, SSearchBox )
 			.HintText(NSLOCTEXT("PlacementMode", "SearchPlaceables", "Search Classes"))
 			.OnTextChanged(this, &SPlacementModeTools::OnSearchChanged)
 		]
@@ -949,7 +969,7 @@ void SPlacementModeTools::RefreshPlaceables()
 	PlaceablesContainer->ClearChildren();
 	SearchResultsContainer->ClearChildren();
 
-	if ( SearchText.IsEmpty() )
+	if ( SearchTextFilter->GetRawFilterText().IsEmpty() )
 	{
 		// Just build the full list with no filtering
 		for ( int32 i = 0; i < PlaceableClassWidgets.Num(); i++ )
@@ -965,8 +985,7 @@ void SPlacementModeTools::RefreshPlaceables()
 		// Filter out the widgets that don't belong
 		for ( int32 i = 0; i < PlaceableClassWidgets.Num(); i++ )
 		{
-			const FString* SourceString = FTextInspector::GetSourceString(PlaceableClassWidgets[i]->AssetDisplayName);
-			if ( PlaceableClassWidgets[i]->AssetDisplayName.ToString().Contains( SearchText.ToString() ) || ( SourceString && SourceString->Contains( SearchText.ToString() ) ) )
+			if ( SearchTextFilter->PassesFilter( PlaceableClassWidgets[i] ) )
 			{
 				SearchResultsContainer->AddSlot()
 				[
@@ -1019,14 +1038,14 @@ FReply SPlacementModeTools::OnKeyDown( const FGeometry& MyGeometry, const FKeyEv
 
 int32 SPlacementModeTools::GetSelectedPanel() const
 {
-	return SearchText.IsEmpty() ? 0 : 1;
+	return SearchTextFilter->GetRawFilterText().IsEmpty() ? 0 : 1;
 }
 
 void SPlacementModeTools::OnSearchChanged(const FText& InFilterText)
 {
-	// If the search text was previously we do a full rebuild of our cached widgets
+	// If the search text was previously empty we do a full rebuild of our cached widgets
 	// for the placeable widgets.
-	if ( SearchText.IsEmpty() )
+	if ( SearchTextFilter->GetRawFilterText().IsEmpty() )
 	{
 		bPlaceablesFullRefreshRequested = true;
 	}
@@ -1035,10 +1054,11 @@ void SPlacementModeTools::OnSearchChanged(const FText& InFilterText)
 		bPlaceablesRefreshRequested = true;
 	}
 
-	SearchText = InFilterText;
+	SearchTextFilter->SetRawFilterText( InFilterText );
+	SearchBoxPtr->SetError( SearchTextFilter->GetFilterErrorText() );
 }
 
 FText SPlacementModeTools::GetHighlightText() const
 {
-	return SearchText;
+	return SearchTextFilter->GetRawFilterText();
 }
diff --git a/Engine/Source/Editor/PlacementMode/Private/SPlacementModeTools.h b/Engine/Source/Editor/PlacementMode/Private/SPlacementModeTools.h
--- a/Engine/Source/Editor/PlacementMode/Private/SPlacementModeTools.h
+++ b/Engine/Source/Editor/PlacementMode/Private/SPlacementModeTools.h
@@ -128,8 +128,12 @@ private:
 	// Flag set when the volumes widget needs to be rebuilt.
 	bool bVolumesRefreshRequested;
 
-	// The last text the user searched
-	FText SearchText;
+	// The text filter used to filter the classes
+	typedef TTextFilter<TSharedRef<SPlacementAssetEntry>> FPlacementAssetEntryTextFilter;
+	TSharedPtr<FPlacementAssetEntryTextFilter> SearchTextFilter;
+
+	// The search box used to update the filter text
+	TSharedPtr<SSearchBox> SearchBoxPtr;
 
 	// The list of placeable classes in widget form, these get put in the search results or
 	// the all classes view depending on which is currently visible.
