Copying //UE4/Dev-Sequencer to Dev-Main (//UE4/Dev-Main)

diff --git a/Engine/Source/Editor/PlacementMode/Private/PlacementModeModule.cpp b/Engine/Source/Editor/PlacementMode/Private/PlacementModeModule.cpp
--- a/Engine/Source/Editor/PlacementMode/Private/PlacementModeModule.cpp
+++ b/Engine/Source/Editor/PlacementMode/Private/PlacementModeModule.cpp
@@ -3,6 +3,50 @@
 #include "PlacementModePrivatePCH.h"
 #include "IPlacementModeModule.h"
 #include "PlacementMode.h"
+#include "AssetToolsModule.h"
+#include "KismetEditorUtilities.h"
+
+
+struct FPlacementCategory : FPlacementCategoryInfo
+{
+	FPlacementCategory(const FPlacementCategoryInfo& SourceInfo)
+		: FPlacementCategoryInfo(SourceInfo)
+	{
+
+	}
+
+	FPlacementCategory(FPlacementCategory&& In)
+		: FPlacementCategoryInfo(MoveTemp(In))
+		, Items(MoveTemp(In.Items))
+	{}
+
+	FPlacementCategory& operator=(FPlacementCategory&& In)
+	{
+		FPlacementCategoryInfo::operator=(MoveTemp(In));
+		Items = MoveTemp(In.Items);
+		return *this;
+	}
+
+	TMap<uint32, TSharedPtr<FPlaceableItem>> Items;
+};
+
+static TOptional<FLinearColor> GetBasicShapeColorOverride()
+{
+	// Get color for basic shapes.  It should appear like all the other basic types
+	static TOptional<FLinearColor> BasicShapeColorOverride;
+	
+	if( !BasicShapeColorOverride.IsSet() )
+	{
+		FAssetToolsModule& AssetToolsModule = FModuleManager::LoadModuleChecked<FAssetToolsModule>(TEXT("AssetTools"));
+		TSharedPtr<IAssetTypeActions> AssetTypeActions;
+		AssetTypeActions = AssetToolsModule.Get().GetAssetTypeActionsForClass(UClass::StaticClass()).Pin();
+		if (AssetTypeActions.IsValid())
+		{
+			BasicShapeColorOverride = TOptional<FLinearColor>(AssetTypeActions->GetTypeColor());
+		}
+	}
+	return BasicShapeColorOverride;
+}
 
 class FPlacementModeModule : public IPlacementModeModule
 {
@@ -32,9 +76,126 @@ public:
 			FSlateIcon(FEditorStyle::GetStyleSetName(), "LevelEditor.PlacementMode", "LevelEditor.PlacementMode.Small"),
 			true, 0);
 
-		FAssetRegistryModule& AssetRegistryModule = FModuleManager::GetModuleChecked<FAssetRegistryModule>(TEXT("AssetRegistry"));
+		FAssetRegistryModule& AssetRegistryModule = FModuleManager::LoadModuleChecked<FAssetRegistryModule>(TEXT("AssetRegistry"));
 		AssetRegistryModule.Get().OnAssetRemoved().AddRaw( this, &FPlacementModeModule::OnAssetRemoved );
 		AssetRegistryModule.Get().OnAssetRenamed().AddRaw( this, &FPlacementModeModule::OnAssetRenamed );
+
+		TOptional<FLinearColor> BasicShapeColorOverride = GetBasicShapeColorOverride();
+
+
+		RegisterPlacementCategory(
+			FPlacementCategoryInfo(
+				NSLOCTEXT( "PlacementMode", "RecentlyPlaced", "Recently Placed" ),
+				FBuiltInPlacementCategories::RecentlyPlaced(),
+				TEXT("PMRecentlyPlaced"),
+				TNumericLimits<int32>::Lowest(),
+				false
+			)
+		);
+
+		{
+			int32 SortOrder = 0;
+			FName CategoryName = FBuiltInPlacementCategories::Basic();
+			RegisterPlacementCategory(
+				FPlacementCategoryInfo(
+					NSLOCTEXT( "PlacementMode", "Basic", "Basic" ),
+					CategoryName,
+					TEXT("PMBasic"),
+					10
+				)
+			);
+
+			FPlacementCategory* Category = Categories.Find(CategoryName);
+			Category->Items.Add(CreateID(), MakeShareable( new FPlaceableItem(*UActorFactoryEmptyActor::StaticClass(), SortOrder+=10)) );
+			Category->Items.Add(CreateID(), MakeShareable( new FPlaceableItem(*UActorFactoryCharacter::StaticClass(), SortOrder+=10)) );
+			Category->Items.Add(CreateID(), MakeShareable( new FPlaceableItem(*UActorFactoryPawn::StaticClass(), SortOrder+=10)) );
+			Category->Items.Add(CreateID(), MakeShareable( new FPlaceableItem(*UActorFactoryPointLight::StaticClass(), SortOrder+=10)) );
+			Category->Items.Add(CreateID(), MakeShareable( new FPlaceableItem(*UActorFactoryPlayerStart::StaticClass(), SortOrder+=10)) );
+			// Cube
+			Category->Items.Add(CreateID(), MakeShareable( new FPlaceableItem(*UActorFactoryBasicShape::StaticClass(), FAssetData(LoadObject<UStaticMesh>(nullptr,*UActorFactoryBasicShape::BasicCube.ToString())), FName("ClassThumbnail.Cube"), BasicShapeColorOverride, SortOrder+=10)) );
+			// Sphere
+			Category->Items.Add(CreateID(), MakeShareable( new FPlaceableItem(*UActorFactoryBasicShape::StaticClass(), FAssetData(LoadObject<UStaticMesh>(nullptr,*UActorFactoryBasicShape::BasicSphere.ToString())), FName("ClassThumbnail.Sphere"), BasicShapeColorOverride, SortOrder+=10 )) );
+			// Cylinder
+			Category->Items.Add(CreateID(), MakeShareable( new FPlaceableItem(*UActorFactoryBasicShape::StaticClass(), FAssetData(LoadObject<UStaticMesh>(nullptr,*UActorFactoryBasicShape::BasicCylinder.ToString())), FName("ClassThumbnail.Cylinder"), BasicShapeColorOverride, SortOrder+=10 )) );
+			// Cone
+			Category->Items.Add(CreateID(), MakeShareable( new FPlaceableItem(*UActorFactoryBasicShape::StaticClass(), FAssetData(LoadObject<UStaticMesh>(nullptr,*UActorFactoryBasicShape::BasicCone.ToString())), FName("ClassThumbnail.Cone"), BasicShapeColorOverride, SortOrder+=10 )) );
+			Category->Items.Add(CreateID(), MakeShareable( new FPlaceableItem(*UActorFactoryTriggerBox::StaticClass(), SortOrder+=10)) );
+			Category->Items.Add(CreateID(), MakeShareable( new FPlaceableItem(*UActorFactoryTriggerSphere::StaticClass(), SortOrder+=10)) );
+		}
+
+		{
+			int32 SortOrder = 0;
+			FName CategoryName = FBuiltInPlacementCategories::Lights();
+			RegisterPlacementCategory(
+				FPlacementCategoryInfo(
+					NSLOCTEXT( "PlacementMode", "Lights", "Lights" ),
+					CategoryName,
+					TEXT("PMLights"),
+					20
+				)
+			);
+
+			FPlacementCategory* Category = Categories.Find(CategoryName);
+			Category->Items.Add(CreateID(), MakeShareable( new FPlaceableItem(*UActorFactoryDirectionalLight::StaticClass(), SortOrder+=10)) );
+			Category->Items.Add(CreateID(), MakeShareable( new FPlaceableItem(*UActorFactoryPointLight::StaticClass(), SortOrder+=10)) );
+			Category->Items.Add(CreateID(), MakeShareable( new FPlaceableItem(*UActorFactorySpotLight::StaticClass(), SortOrder+=10)) );
+			Category->Items.Add(CreateID(), MakeShareable( new FPlaceableItem(*UActorFactorySkyLight::StaticClass(), SortOrder+=10)) );
+		}
+
+		{
+			int32 SortOrder = 0;
+			FName CategoryName = FBuiltInPlacementCategories::Visual();
+			RegisterPlacementCategory(
+				FPlacementCategoryInfo(
+					NSLOCTEXT( "PlacementMode", "VisualEffects", "Visual Effects" ),
+					CategoryName,
+					TEXT("PMVisual"),
+					30
+				)
+			);
+
+			UActorFactory* PPFactory = GEditor->FindActorFactoryByClassForActorClass( UActorFactoryBoxVolume::StaticClass(), APostProcessVolume::StaticClass() );
+
+			FPlacementCategory* Category = Categories.Find(CategoryName);
+			Category->Items.Add(CreateID(), MakeShareable( new FPlaceableItem(PPFactory, FAssetData(APostProcessVolume::StaticClass()), SortOrder+=10)) );
+			Category->Items.Add(CreateID(), MakeShareable( new FPlaceableItem(*UActorFactoryAtmosphericFog::StaticClass(), SortOrder+=10)) );
+			Category->Items.Add(CreateID(), MakeShareable( new FPlaceableItem(*UActorFactoryExponentialHeightFog::StaticClass(), SortOrder+=10)) );
+			Category->Items.Add(CreateID(), MakeShareable( new FPlaceableItem(*UActorFactorySphereReflectionCapture::StaticClass(), SortOrder+=10)) );
+			Category->Items.Add(CreateID(), MakeShareable( new FPlaceableItem(*UActorFactoryBoxReflectionCapture::StaticClass(), SortOrder+=10)) );
+			Category->Items.Add(CreateID(), MakeShareable( new FPlaceableItem(*UActorFactoryDeferredDecal::StaticClass(), SortOrder+=10)) );
+		}
+
+		RegisterPlacementCategory(
+			FPlacementCategoryInfo(
+				NSLOCTEXT( "PlacementMode", "Volumes", "Volumes" ),
+				"Volumes",
+				TEXT("PMVolumes"),
+				40
+			)
+		);
+
+		{
+			int32 SortOrder = 0;
+			FName CategoryName = FBuiltInPlacementCategories::AllClasses();
+			RegisterPlacementCategory(
+				FPlacementCategoryInfo(
+					NSLOCTEXT( "PlacementMode", "AllClasses", "All Classes" ),
+					CategoryName,
+					TEXT("PMAllClasses"),
+					50
+				)
+			);
+
+			FPlacementCategory* Category = Categories.Find(CategoryName);
+			Category->Items.Add(CreateID(), MakeShareable( new FPlaceableItem(*UActorFactoryEmptyActor::StaticClass())) );
+			Category->Items.Add(CreateID(), MakeShareable( new FPlaceableItem(*UActorFactoryCharacter::StaticClass())) );
+			Category->Items.Add(CreateID(), MakeShareable( new FPlaceableItem(*UActorFactoryPawn::StaticClass())) );
+			Category->Items.Add(CreateID(), MakeShareable( new FPlaceableItem(*UActorFactoryBasicShape::StaticClass(), FAssetData(LoadObject<UStaticMesh>(nullptr, *UActorFactoryBasicShape::BasicCube.ToString())), FName("ClassThumbnail.Cube"), GetBasicShapeColorOverride())) );
+			Category->Items.Add(CreateID(), MakeShareable( new FPlaceableItem(*UActorFactoryBasicShape::StaticClass(), FAssetData(LoadObject<UStaticMesh>(nullptr, *UActorFactoryBasicShape::BasicSphere.ToString())), FName("ClassThumbnail.Sphere"), GetBasicShapeColorOverride())) );
+			Category->Items.Add(CreateID(), MakeShareable( new FPlaceableItem(*UActorFactoryBasicShape::StaticClass(), FAssetData(LoadObject<UStaticMesh>(nullptr, *UActorFactoryBasicShape::BasicCylinder.ToString())), FName("ClassThumbnail.Cylinder"), GetBasicShapeColorOverride())) );
+			Category->Items.Add(CreateID(), MakeShareable( new FPlaceableItem(*UActorFactoryBasicShape::StaticClass(), FAssetData(LoadObject<UStaticMesh>(nullptr, *UActorFactoryBasicShape::BasicCone.ToString())), FName("ClassThumbnail.Cone"), GetBasicShapeColorOverride())) );
+		}
+
 	}
 
 	/**
@@ -179,8 +340,262 @@ public:
 	{
 		StoppedPlacingEvent.Broadcast( bWasSuccessfullyPlaced );
 	}
+
+public:
+
+	virtual bool RegisterPlacementCategory(const FPlacementCategoryInfo& Info)
+	{
+		if (Categories.Contains(Info.UniqueHandle))
+		{
+			return false;
+		}
+
+		Categories.Add(Info.UniqueHandle, Info);
+		return true;
+	}
+
+	virtual const FPlacementCategoryInfo* GetRegisteredPlacementCategory(FName CategoryName) const override
+	{
+		return Categories.Find(CategoryName);
+	}
+
+	virtual void UnregisterPlacementCategory(FName Handle)
+	{
+		Categories.Remove(Handle);
+	}
+
+	virtual void GetSortedCategories(TArray<FPlacementCategoryInfo>& OutCategories) const
+	{
+		TArray<FName> SortedNames;
+		Categories.GenerateKeyArray(SortedNames);
+
+		SortedNames.Sort([&](const FName& A, const FName& B){
+			return Categories[A].SortOrder < Categories[B].SortOrder;
+		});
+
+		OutCategories.Reset(Categories.Num());
+		for (const FName& Name : SortedNames)
+		{
+			OutCategories.Add(Categories[Name]);
+		}
+	}
+
+	virtual TOptional<FPlacementModeID> RegisterPlaceableItem(FName CategoryName, const TSharedRef<FPlaceableItem>& InItem)
+	{
+		FPlacementCategory* Category = Categories.Find(CategoryName);
+		if (Category && !Category->CustomGenerator)
+		{
+			FPlacementModeID ID = CreateID(CategoryName);
+			Category->Items.Add(ID.UniqueID, InItem);
+			return ID;
+		}
+		return TOptional<FPlacementModeID>();
+	}
+
+	virtual void UnregisterPlaceableItem(FPlacementModeID ID)
+	{
+		FPlacementCategory* Category = Categories.Find(ID.Category);
+		if (Category)
+		{
+			Category->Items.Remove(ID.UniqueID);
+		}
+	}
+
+	virtual void GetItemsForCategory(FName CategoryName, TArray<TSharedPtr<FPlaceableItem>>& OutItems) const
+	{
+		const FPlacementCategory* Category = Categories.Find(CategoryName);
+		if (Category)
+		{
+			for (auto& Pair : Category->Items)
+			{
+				OutItems.Add(Pair.Value);
+			}
+		}
+	}
+
+	virtual void GetFilteredItemsForCategory(FName CategoryName, TArray<TSharedPtr<FPlaceableItem>>& OutItems, TFunctionRef<bool(const TSharedPtr<FPlaceableItem>&)> Filter) const
+	{
+		const FPlacementCategory* Category = Categories.Find(CategoryName);
+		if (Category)
+		{
+			for (auto& Pair : Category->Items)
+			{
+				if (Filter(Pair.Value))
+				{
+					OutItems.Add(Pair.Value);
+				}
+			}
+		}
+	}
+
+	virtual void RegenerateItemsForCategory(FName Category) override
+	{
+		if (Category == FBuiltInPlacementCategories::RecentlyPlaced())
+		{
+			RefreshRecentlyPlaced();
+		}
+		else if (Category == FBuiltInPlacementCategories::Volumes())
+		{
+			RefreshVolumes();
+		}
+		else if (Category == FBuiltInPlacementCategories::AllClasses())
+		{
+			RefreshAllPlaceableClasses();
+		}
+
+		// @todo: Broadcast so clients can update dynamic lists?
+	}
+
+	void RefreshRecentlyPlaced()
+	{
+		FName CategoryName = FBuiltInPlacementCategories::RecentlyPlaced();
+
+		FPlacementCategory* Category = Categories.Find(CategoryName);
+		if (!Category)
+		{
+			return;
+		}
+
+		Category->Items.Reset();
+
+
+		FAssetRegistryModule& AssetRegistryModule = FModuleManager::LoadModuleChecked<FAssetRegistryModule>( TEXT( "AssetRegistry" ) );
+
+		for (const FActorPlacementInfo& RecentlyPlacedItem : RecentlyPlaced)
+		{
+			UObject* Asset = FindObject<UObject>(nullptr, *RecentlyPlacedItem.ObjectPath);
+
+			// If asset is pending delete, it will not be marked as RF_Standalone, in which case we skip it
+			if (Asset != nullptr && Asset->HasAnyFlags(RF_Standalone))
+			{
+				FAssetData AssetData = AssetRegistryModule.Get().GetAssetByObjectPath(*RecentlyPlacedItem.ObjectPath);
+
+				if (AssetData.IsValid())
+				{
+					UActorFactory* Factory = FindObject<UActorFactory>(nullptr, *RecentlyPlacedItem.Factory);
+					Category->Items.Add(CreateID(), MakeShareable( new FPlaceableItem(Factory, AssetData) ));
+				}
+			}
+		}
+
+		// Broadcast?
+	}
+
+	void RefreshVolumes()
+	{
+		FName CategoryName = FBuiltInPlacementCategories::Volumes();
+
+		FPlacementCategory* Category = Categories.Find(CategoryName);
+		if (!Category)
+		{
+			return;
+		}
+
+		Category->Items.Reset();
+
+		// Add loaded classes
+		for ( TObjectIterator<UClass> ClassIt; ClassIt; ++ClassIt )
+		{
+			const UClass* Class = *ClassIt;
+
+			if (!Class->HasAllClassFlags(CLASS_NotPlaceable) &&
+				!Class->HasAnyClassFlags(CLASS_Abstract | CLASS_Deprecated | CLASS_NewerVersionExists) &&
+				Class->IsChildOf(AVolume::StaticClass()) &&
+				Class->ClassGeneratedBy == nullptr )
+			{
+				UActorFactory* Factory = GEditor->FindActorFactoryByClassForActorClass(UActorFactoryBoxVolume::StaticClass(), Class);
+				Category->Items.Add(CreateID(), MakeShareable( new FPlaceableItem(Factory, FAssetData(Class)) ));
+			}
+		}
+
+		// Broadcast?
+	}
+
+	void RefreshAllPlaceableClasses()
+	{
+		FName CategoryName = FBuiltInPlacementCategories::AllClasses();
+
+		// Unregister old stuff
+		FPlacementCategory* Category = Categories.Find(CategoryName);
+		if (!Category)
+		{
+			return;
+		}
+
+		Category->Items.Reset();
+
+		// Make a map of UClasses to ActorFactories that support them
+		const TArray< UActorFactory *>& ActorFactories = GEditor->ActorFactories;
+		TMap<UClass*, UActorFactory*> ActorFactoryMap;
+		for ( int32 FactoryIdx = 0; FactoryIdx < ActorFactories.Num(); ++FactoryIdx )
+		{
+			UActorFactory* ActorFactory = ActorFactories[FactoryIdx];
+
+			if ( ActorFactory )
+			{
+				ActorFactoryMap.Add(ActorFactory->GetDefaultActorClass(FAssetData()), ActorFactory);
+			}
+		}
+
+		FAssetData NoAssetData;
+		FText UnusedErrorMessage;
+
+		// Add loaded classes
+		for ( TObjectIterator<UClass> ClassIt; ClassIt; ++ClassIt )
+		{
+			// Don't offer skeleton classes
+			bool bIsSkeletonClass = FKismetEditorUtilities::IsClassABlueprintSkeleton(*ClassIt);
+
+			if ( !ClassIt->HasAllClassFlags(CLASS_NotPlaceable) &&
+				!ClassIt->HasAnyClassFlags(CLASS_Abstract | CLASS_Deprecated | CLASS_NewerVersionExists) &&
+				ClassIt->IsChildOf(AActor::StaticClass()) &&
+				( !ClassIt->IsChildOf(ABrush::StaticClass()) || ClassIt->IsChildOf(AVolume::StaticClass()) ) &&
+				!bIsSkeletonClass )
+			{
+				UActorFactory* ActorFactory = ActorFactoryMap.FindRef(*ClassIt);
+
+				const bool IsVolume = ClassIt->IsChildOf(AVolume::StaticClass());
+
+				if ( IsVolume )
+				{
+					ActorFactory = GEditor->FindActorFactoryByClassForActorClass(UActorFactoryBoxVolume::StaticClass(), *ClassIt);
+				}
+				else if (ActorFactory && !ActorFactory->CanCreateActorFrom(NoAssetData, UnusedErrorMessage))
+				{
+					continue;
+				}
+
+				Category->Items.Add(CreateID(), MakeShareable( new FPlaceableItem(ActorFactory, FAssetData(*ClassIt)) ));
+			}
+		}
+
+		// Broadcast?
+	}
+
+private:
+
+	uint32 CreateID()
+	{
+		static uint32 IDCounter = 0;
+		uint32 Next = ++IDCounter;
+		checkf(Next, TEXT("uint32 overflow. If this occurs, something's gone seriously wrong."))
+		return Next;
+	}
+
+	FPlacementModeID CreateID(FName InCategory)
+	{
+		FPlacementModeID NewID;
+		NewID.UniqueID = CreateID();
+		NewID.Category = InCategory;
+
+		
+		return NewID;
+	}
+
 private:
 
+	TMap<FName, FPlacementCategory> Categories;
+
 	TArray< FActorPlacementInfo >	RecentlyPlaced;
 	FOnRecentlyPlacedChanged		RecentlyPlacedChanged;
 
diff --git a/Engine/Source/Editor/PlacementMode/Private/SPlacementModeTools.cpp b/Engine/Source/Editor/PlacementMode/Private/SPlacementModeTools.cpp
--- a/Engine/Source/Editor/PlacementMode/Private/SPlacementModeTools.cpp
+++ b/Engine/Source/Editor/PlacementMode/Private/SPlacementModeTools.cpp
@@ -7,46 +7,56 @@
 #include "SPlacementModeTools.h"
 #include "DragAndDrop/AssetDragDropOp.h"
 #include "IBspModeModule.h"
-#include "KismetEditorUtilities.h"
 #include "EditorClassUtils.h"
 #include "SSearchBox.h"
-#include "SWidgetSwitcher.h"
 #include "GameFramework/Volume.h"
 #include "Engine/PostProcessVolume.h"
 #include "AssetToolsModule.h"
 
-/**
- * These are the tab indexes, if the tabs are reorganized you need to adjust the
- * enum accordingly.
- */
-namespace EPlacementTab
+
+struct FSortPlaceableItems
 {
-	enum Type
+	static bool SortItemsByOrderThenName(const TSharedPtr<FPlaceableItem>& A, const TSharedPtr<FPlaceableItem>& B)
 	{
-		RecentlyPlaced = 0,
-		Basic,
-		Lights,
-		Visual,
-		Geometry,
-		Volumes,
-		AllClasses,
-	};
-}
+		if (A->SortOrder.IsSet())
+		{
+			if (B->SortOrder.IsSet())
+			{
+				return A->SortOrder.GetValue() < B->SortOrder.GetValue();
+			}
+			else
+			{
+				return true;
+			}
+		}
+		else if (B->SortOrder.IsSet())
+		{
+			return false;
+		}
+		else
+		{
+			return SortItemsByName(A, B);
+		}
+	}
 
-namespace PlacementViewFilter
-{
+	static bool SortItemsByName(const TSharedPtr<FPlaceableItem>& A, const TSharedPtr<FPlaceableItem>& B)
+	{
+		return A->DisplayName.CompareTo(B->DisplayName) < 0;
+	}
+};
 
-void GetBasicStrings(TSharedRef<SPlacementAssetEntry> InPlacementAsset, TArray<FString>& OutBasicStrings)
+namespace PlacementViewFilter
 {
-	OutBasicStrings.Add(InPlacementAsset->AssetDisplayName.ToString());
-
-	const FString* SourceString = FTextInspector::GetSourceString(InPlacementAsset->AssetDisplayName);
-	if (SourceString)
+	void GetBasicStrings(const FPlaceableItem& InPlaceableItem, TArray<FString>& OutBasicStrings)
 	{
-		OutBasicStrings.Add(*SourceString);
-	}
-}
+		OutBasicStrings.Add(InPlaceableItem.DisplayName.ToString());
 
+		const FString* SourceString = FTextInspector::GetSourceString(InPlaceableItem.DisplayName);
+		if (SourceString)
+		{
+			OutBasicStrings.Add(*SourceString);
+		}
+	}
 } // namespace PlacementViewFilter
 
 /**
@@ -98,37 +108,27 @@ private:
 	TSharedPtr< FAssetThumbnail > Thumbnail;
 };
 
-void SPlacementAssetEntry::Construct(const FArguments& InArgs, UActorFactory* InFactory, const FAssetData& InAsset)
+void SPlacementAssetEntry::Construct(const FArguments& InArgs, const TSharedPtr<FPlaceableItem>& InItem)
 {	
 	bIsPressed = false;
 
-	FactoryToUse = InFactory;
-	AssetData = InAsset;
+	Item = InItem;
 
 	TSharedPtr< SHorizontalBox > ActorType = SNew( SHorizontalBox );
 
-	const bool IsClass = AssetData.GetClass() == UClass::StaticClass();
-	const bool IsVolume = IsClass ? CastChecked<UClass>( AssetData.GetAsset() )->IsChildOf( AVolume::StaticClass() ) : false;
-
-	if ( IsClass )
-	{
-		AssetDisplayName = FText::FromString( FName::NameToDisplayString( AssetData.AssetName.ToString(), false ) );
-	}
-	else
-	{
-		AssetDisplayName = FText::FromName( AssetData.AssetName );
-	}
+	const bool IsClass = Item->AssetData.GetClass() == UClass::StaticClass();
+	const bool IsVolume = IsClass ? CastChecked<UClass>( Item->AssetData.GetAsset() )->IsChildOf( AVolume::StaticClass() ) : false;
 
 	FText ActorTypeDisplayName;
 	AActor* DefaultActor = nullptr;
-	if ( FactoryToUse != nullptr )
+	if ( Item->Factory != nullptr )
 	{
-		DefaultActor = FactoryToUse->GetDefaultActor( AssetData );
-		ActorTypeDisplayName = FactoryToUse->GetDisplayName();
+		DefaultActor = Item->Factory->GetDefaultActor( Item->AssetData );
+		ActorTypeDisplayName = Item->Factory->GetDisplayName();
 	}
-	else if ( IsClass && CastChecked<UClass>( AssetData.GetAsset() )->IsChildOf( AActor::StaticClass() ) )
+	else if ( IsClass && CastChecked<UClass>( Item->AssetData.GetAsset() )->IsChildOf( AActor::StaticClass() ) )
 	{
-		DefaultActor = CastChecked<AActor>( CastChecked<UClass>( AssetData.GetAsset() )->ClassDefaultObject );
+		DefaultActor = CastChecked<AActor>( CastChecked<UClass>( Item->AssetData.GetAsset() )->ClassDefaultObject );
 		ActorTypeDisplayName = FText::FromString( FName::NameToDisplayString( DefaultActor->GetClass()->GetName(), false ) );
 	}
 	
@@ -140,21 +140,14 @@ void SPlacementAssetEntry::Construct(const FArguments& InArgs, UActorFactory* In
 		ToolTip = FEditorClassUtils::GetTooltip(DefaultActor->GetClass());
 	}
 
-	if ( InArgs._LabelOverride.IsEmpty() )
-	{
-		if (IsClass && !IsVolume && !ActorTypeDisplayName.IsEmpty())
-		{
-			AssetDisplayName = ActorTypeDisplayName;
-		}
-	}
-	else
+	if (IsClass && !IsVolume && !ActorTypeDisplayName.IsEmpty())
 	{
-		AssetDisplayName = InArgs._LabelOverride;
+		Item->DisplayName = ActorTypeDisplayName;
 	}
 
 	if (!ToolTip.IsValid())
 	{
-		ToolTip = FSlateApplicationBase::Get().MakeToolTip(AssetDisplayName);
+		ToolTip = FSlateApplicationBase::Get().MakeToolTip(Item->DisplayName);
 	}
 	
 	const FButtonStyle& ButtonStyle = FEditorStyle::GetWidgetStyle<FButtonStyle>( "PlacementBrowser.Asset" );
@@ -193,10 +186,10 @@ void SPlacementAssetEntry::Construct(const FArguments& InArgs, UActorFactory* In
 					.WidthOverride( 35 )
 					.HeightOverride( 35 )
 					[
-						SNew( SPlacementAssetThumbnail, AssetData )
-						.ClassThumbnailBrushOverride( InArgs._ClassThumbnailBrushOverride )
-						.AlwaysUseGenericThumbnail( InArgs._AlwaysUseGenericThumbnail )
-						.AssetTypeColorOverride( InArgs._AssetTypeColorOverride )
+						SNew( SPlacementAssetThumbnail, Item->AssetData )
+						.ClassThumbnailBrushOverride( Item->ClassThumbnailBrushOverride )
+						.AlwaysUseGenericThumbnail( Item->bAlwaysUseGenericThumbnail )
+						.AssetTypeColorOverride( Item->AssetTypeColorOverride )
 					]
 				]
 			]
@@ -212,7 +205,7 @@ void SPlacementAssetEntry::Construct(const FArguments& InArgs, UActorFactory* In
 				[
 					SNew( STextBlock )
 					.TextStyle( FEditorStyle::Get(), "PlacementBrowser.Asset.Name" )
-					.Text( AssetDisplayName )
+					.Text( Item->DisplayName )
 					.HighlightText(InArgs._HighlightText)
 				]
 			]
@@ -253,7 +246,7 @@ FReply SPlacementAssetEntry::OnDragDetected(const FGeometry& MyGeometry, const F
 {
 	bIsPressed = false;
 
-	return FReply::Handled().BeginDragDrop( FAssetDragDropOp::New( AssetData, FactoryToUse ) );
+	return FReply::Handled().BeginDragDrop( FAssetDragDropOp::New( Item->AssetData, Item->Factory ) );
 }
 
 bool SPlacementAssetEntry::IsPressed() const
@@ -287,9 +280,9 @@ SPlacementModeTools::~SPlacementModeTools()
 
 void SPlacementModeTools::Construct( const FArguments& InArgs )
 {
-	bPlaceablesRefreshRequested = false;
 	bPlaceablesFullRefreshRequested = false;
-	bVolumesRefreshRequested = false;
+	bRecentlyPlacedRefreshRequested = false;
+	bNeedsUpdate = true;
 
 	FPlacementMode* PlacementEditMode = (FPlacementMode*)GLevelEditorModeTools().GetActiveMode( FBuiltinEditorModes::EM_Placement );
 	PlacementEditMode->AddValidFocusTargetForPlacement( SharedThis( this ) );
@@ -298,6 +291,25 @@ void SPlacementModeTools::Construct( const FArguments& InArgs )
 		FPlacementAssetEntryTextFilter::FItemToStringArray::CreateStatic(&PlacementViewFilter::GetBasicStrings)
 		));
 
+	TSharedRef<SVerticalBox> Tabs = SNew(SVerticalBox).Visibility(this, &SPlacementModeTools::GetTabsVisibility);
+
+	// Populate the tabs and body from the defined placeable items
+	IPlacementModeModule& PlacementModeModule = IPlacementModeModule::Get();
+
+	TArray<FPlacementCategoryInfo> Categories;
+	PlacementModeModule.GetSortedCategories(Categories);
+	for (const FPlacementCategoryInfo& Category : Categories)
+	{
+		Tabs->AddSlot()
+		.AutoHeight()
+		[
+			CreatePlacementGroupTab(Category)
+		];
+	}
+
+	TSharedRef<SScrollBar> ScrollBar = SNew(SScrollBar)
+		.Thickness(FVector2D(5, 5));
+
 	ChildSlot
 	[
 		SNew( SVerticalBox )
@@ -313,224 +325,73 @@ void SPlacementModeTools::Construct( const FArguments& InArgs )
 
 		+ SVerticalBox::Slot()
 		.Padding( 0 )
-		.FillHeight( 1.0f )
 		[
-			SNew(SWidgetSwitcher)
-			.WidgetIndex(this, &SPlacementModeTools::GetSelectedPanel)
+			SNew(SHorizontalBox)
 
-			// Normal Panel
-			+ SWidgetSwitcher::Slot()
+			+ SHorizontalBox::Slot()
+			.AutoWidth()
 			[
-				CreateStandardPanel()
+				Tabs
 			]
 
-			// Search Results Panel
-			+ SWidgetSwitcher::Slot()
-			.Padding(4, 0, 4, 4)
+			+ SHorizontalBox::Slot()
 			[
 				SNew(SBorder)
 				.Padding(FMargin(3))
 				.BorderImage(FEditorStyle::GetBrush("ToolPanel.DarkGroupBorder"))
 				[
-					SNew(SScrollBox)
+					SNew(SOverlay)
 
-					+ SScrollBox::Slot()
+					+ SOverlay::Slot()
+					.HAlign(HAlign_Center)
+					.VAlign(VAlign_Fill)
 					[
-						SAssignNew(SearchResultsContainer, SVerticalBox)
+						SNew(STextBlock)
+						.Text(NSLOCTEXT("PlacementMode", "NoResultsFound", "No Results Found"))
+						.Visibility(this, &SPlacementModeTools::GetFailedSearchVisibility)
 					]
-				]
-			]
-		]
-	];
-
-	RefreshRecentlyPlaced();
-
-	IPlacementModeModule::Get().OnRecentlyPlacedChanged().AddSP( this, &SPlacementModeTools::UpdateRecentlyPlacedAssets );
-}
-
-TSharedRef< SWidget > SPlacementModeTools::CreateStandardPanel()
-{
-	return SNew( SHorizontalBox )
-
-	// The tabs on the left
-	+ SHorizontalBox::Slot()
-	.AutoWidth()
-	[
-		SNew( SVerticalBox )
-
-		+ SVerticalBox::Slot()
-		.Padding( 0, 3, 0, 0 )
-		.AutoHeight()
-		[
-			SNew( SBox)
-			.AddMetaData<FTagMetaData>(FTagMetaData(TEXT("PMRecentlyPlaced")))
-			[
-				CreatePlacementGroupTab( (int32)EPlacementTab::RecentlyPlaced, NSLOCTEXT( "PlacementMode", "RecentlyPlaced", "Recently Placed" ), true )
-			]
-		]
-
 
-		+ SVerticalBox::Slot()
-		.AutoHeight()
-		[
-			SNew(SBox)
-			.AddMetaData<FTagMetaData>(FTagMetaData(TEXT("PMBasic")))
-			[
-				CreatePlacementGroupTab((int32)EPlacementTab::Basic, NSLOCTEXT("PlacementMode", "Basic", "Basic"), false)
-			]
-		]
-
-		+ SVerticalBox::Slot()
-		.AutoHeight()
-		[
-			SNew( SBox )
-			.AddMetaData<FTagMetaData>(FTagMetaData(TEXT("PMLights")))
-			[
-				CreatePlacementGroupTab( (int32)EPlacementTab::Lights, NSLOCTEXT( "PlacementMode", "Lights", "Lights" ), false )
-			]
-		]
-
-		+ SVerticalBox::Slot()
-		.AutoHeight()
-		[
-			SNew( SBox )
-			.AddMetaData<FTagMetaData>(FTagMetaData(TEXT("PMVisual")))
-			[
-				CreatePlacementGroupTab( (int32)EPlacementTab::Visual, NSLOCTEXT( "PlacementMode", "VisualEffects", "Visual Effects" ), false )
-			]
-		]
-
-
-		+ SVerticalBox::Slot()
-		.AutoHeight()
-		[
-			SNew(SBox)
-			.AddMetaData<FTagMetaData>(FTagMetaData(TEXT("PMGeometry")))
-			[
-				CreatePlacementGroupTab((int32)EPlacementTab::Geometry, NSLOCTEXT("PlacementMode", "BSP", "BSP"), false)
-			]
-		]
-
-		+ SVerticalBox::Slot()
-		.AutoHeight()
-		[
-			SNew( SBox )
-			.AddMetaData<FTagMetaData>(FTagMetaData(TEXT("PMVolumes")))
-			[
-				CreatePlacementGroupTab( (int32)EPlacementTab::Volumes, NSLOCTEXT( "PlacementMode", "Volumes", "Volumes" ), false )
-			]
-		]
-
-		+ SVerticalBox::Slot()
-		.AutoHeight()
-		[
-			SNew( SBox )
-			.AddMetaData<FTagMetaData>(FTagMetaData(TEXT("PMAllClasses")))
-			[
-				CreatePlacementGroupTab( (int32)EPlacementTab::AllClasses, NSLOCTEXT( "PlacementMode", "AllClasses", "All Classes" ), true )
-			]
-		]
-	]
-
-	// The 'tab body' area that is switched out with the widget switcher based on the currently active tab.
-	+ SHorizontalBox::Slot()
-	.FillWidth( 1.0f )
-	[
-		SNew( SVerticalBox )
-
-		+ SVerticalBox::Slot()
-		.FillHeight( 1.0f )
-		.Padding( 0 )
-		[
-			SNew( SBorder )
-			.Padding( FMargin( 3 ) )
-			.BorderImage( FEditorStyle::GetBrush( "ToolPanel.DarkGroupBorder" ) )
-			[
-				SAssignNew( WidgetSwitcher, SWidgetSwitcher )
-				.WidgetIndex( (int32)EPlacementTab::Basic )
-
-				// Recently Placed
-				+ SWidgetSwitcher::Slot()
-				[
-					SNew( SScrollBox )
-
-					+ SScrollBox::Slot()
-					[
-						SAssignNew( RecentlyPlacedContainer, SVerticalBox )
-					]
-				]
-
-				// Basics
-				+ SWidgetSwitcher::Slot()
-				[
-					SNew(SScrollBox)
-
-					+ SScrollBox::Slot()
-					[
-						BuildBasicWidget()
-					]
-				]
-
-				// Lights
-				+ SWidgetSwitcher::Slot()
-				[
-					SNew( SScrollBox )
-							
-					+ SScrollBox::Slot()
-					[
-						BuildLightsWidget()
-					]
-				]
-
-				// Visual Effects
-				+ SWidgetSwitcher::Slot()
-				[
-					SNew( SScrollBox )
-							
-					+ SScrollBox::Slot()
-					[
-						BuildVisualWidget()
-					]
-				]
-
-				// Geometry
-				+SWidgetSwitcher::Slot()
-				[
-					FModuleManager::LoadModuleChecked<IBspModeModule>("BspMode").CreateBspModeWidget()
-				]
-
-				// Volumes
-				+ SWidgetSwitcher::Slot()
-				[
-					SNew( SScrollBox )
-
-					+ SScrollBox::Slot()
+					+ SOverlay::Slot()
 					[
-						SAssignNew( VolumesContainer, SVerticalBox )
+						SAssignNew(CustomContent, SBox)
 					]
-				]
-
-				// Classes
-				+ SWidgetSwitcher::Slot()
-				[
-					SNew( SScrollBox )
 
-					+ SScrollBox::Slot()
+					+ SOverlay::Slot()
 					[
-						SAssignNew( PlaceablesContainer, SVerticalBox )
+						SNew(SHorizontalBox)
+
+						+ SHorizontalBox::Slot()
+						[
+							SAssignNew(ListView, SListView<TSharedPtr<FPlaceableItem>>)
+							.ListItemsSource(&FilteredItems)
+							.OnGenerateRow(this, &SPlacementModeTools::OnGenerateWidgetForItem)
+							.Visibility(this, &SPlacementModeTools::GetListViewVisibility)
+							.ExternalScrollbar(ScrollBar)
+						]
+
+						+ SHorizontalBox::Slot()
+						.AutoWidth()
+						[
+							ScrollBar
+						]
 					]
 				]
 			]
 		]
 	];
+
+	ActiveTabName = FBuiltInPlacementCategories::Basic();
+	bNeedsUpdate = true;
+
+	PlacementModeModule.OnRecentlyPlacedChanged().AddSP( this, &SPlacementModeTools::UpdateRecentlyPlacedAssets );
 }
 
-TSharedRef< SWidget > SPlacementModeTools::CreatePlacementGroupTab( int32 TabIndex, FText TabText, bool Important )
+TSharedRef< SWidget > SPlacementModeTools::CreatePlacementGroupTab( const FPlacementCategoryInfo& Info )
 {
 	return SNew( SCheckBox )
 	.Style( FEditorStyle::Get(), "PlacementBrowser.Tab" )
-	.OnCheckStateChanged( this, &SPlacementModeTools::OnPlacementTabChanged, TabIndex )
-	.IsChecked( this, &SPlacementModeTools::GetPlacementTabCheckedState, TabIndex )
+	.OnCheckStateChanged( this, &SPlacementModeTools::OnPlacementTabChanged, Info.UniqueHandle )
+	.IsChecked( this, &SPlacementModeTools::GetPlacementTabCheckedState, Info.UniqueHandle )
 	[
 		SNew( SOverlay )
 
@@ -546,8 +407,8 @@ TSharedRef< SWidget > SPlacementModeTools::CreatePlacementGroupTab( int32 TabInd
 		.VAlign( VAlign_Center )
 		[
 			SNew( STextBlock )
-			.TextStyle( FEditorStyle::Get(), Important ? "PlacementBrowser.Tab.ImportantText" : "PlacementBrowser.Tab.Text" )
-			.Text( TabText )
+			.TextStyle( FEditorStyle::Get(), "PlacementBrowser.Tab.Text" )
+			.Text( Info.DisplayName )
 		]
 
 		+ SOverlay::Slot()
@@ -555,470 +416,163 @@ TSharedRef< SWidget > SPlacementModeTools::CreatePlacementGroupTab( int32 TabInd
 		.HAlign( HAlign_Left )
 		[
 			SNew(SImage)
-			.Image( this, &SPlacementModeTools::PlacementGroupBorderImage, TabIndex )
+			.Image( this, &SPlacementModeTools::PlacementGroupBorderImage, Info.UniqueHandle )
 		]
 	];
 }
 
-ECheckBoxState SPlacementModeTools::GetPlacementTabCheckedState( int32 PlacementGroupIndex ) const
+FName SPlacementModeTools::GetActiveTab() const
 {
-	return WidgetSwitcher->GetActiveWidgetIndex() == PlacementGroupIndex ? ECheckBoxState::Checked : ECheckBoxState::Unchecked;
+	return IsSearchActive() ? FBuiltInPlacementCategories::AllClasses() : ActiveTabName;
 }
 
-void SPlacementModeTools::OnPlacementTabChanged( ECheckBoxState NewState, int32 PlacementGroupIndex )
+void SPlacementModeTools::UpdateFilteredItems()
 {
-	if ( NewState == ECheckBoxState::Checked )
+	bNeedsUpdate = false;
+
+	IPlacementModeModule& PlacementModeModule = IPlacementModeModule::Get();
+
+	const FPlacementCategoryInfo* Category = PlacementModeModule.GetRegisteredPlacementCategory(GetActiveTab());
+	if (!Category)
+	{
+		return;
+	}
+	else if (Category->CustomGenerator)
 	{
-		WidgetSwitcher->SetActiveWidgetIndex( PlacementGroupIndex );
+		CustomContent->SetContent(Category->CustomGenerator());
 
-		if ( PlacementGroupIndex == (int32)EPlacementTab::Volumes )
+		CustomContent->SetVisibility(EVisibility::Visible);
+		ListView->SetVisibility(EVisibility::Collapsed);
+	}
+	else
+	{
+		FilteredItems.Reset();
+
+		const FPlacementCategoryInfo* CategoryInfo = PlacementModeModule.GetRegisteredPlacementCategory(GetActiveTab());
+		if (!ensure(CategoryInfo))
 		{
-			bVolumesRefreshRequested = true;
+			return;
 		}
-		else if ( PlacementGroupIndex == (int32)EPlacementTab::AllClasses )
+		
+		if (IsSearchActive())
 		{
-			bPlaceablesFullRefreshRequested = true;
+			auto Filter = [&](const TSharedPtr<FPlaceableItem>& Item) { return SearchTextFilter->PassesFilter(*Item); };
+			PlacementModeModule.GetFilteredItemsForCategory(CategoryInfo->UniqueHandle, FilteredItems, Filter);
+			
+			if (CategoryInfo->bSortable)
+			{
+				FilteredItems.Sort(&FSortPlaceableItems::SortItemsByName);
+			}
 		}
-	}
-}
+		else
+		{
+			PlacementModeModule.GetItemsForCategory(CategoryInfo->UniqueHandle, FilteredItems);
 
-const FSlateBrush* SPlacementModeTools::PlacementGroupBorderImage( int32 PlacementGroupIndex ) const
-{
-	if ( WidgetSwitcher->GetActiveWidgetIndex() == PlacementGroupIndex )
-	{
-		static FName PlacementBrowserActiveTabBarBrush( "PlacementBrowser.ActiveTabBar" );
-		return FEditorStyle::GetBrush( PlacementBrowserActiveTabBarBrush );
-	}
-	else
-	{
-		return nullptr;
-	}
-}
+			if (CategoryInfo->bSortable)
+			{
+				FilteredItems.Sort(&FSortPlaceableItems::SortItemsByOrderThenName);
+			}
+		}
 
-TSharedRef< SPlacementAssetEntry > BuildDraggableAssetWidget( UClass* InAssetClass, TAttribute<FText> HightlightText = TAttribute<FText>() )
-{
-	UActorFactory* Factory = GEditor->FindActorFactoryByClass( InAssetClass );
-	FAssetData AssetData = FAssetData( Factory->GetDefaultActorClass( FAssetData() ) );
-	return SNew( SPlacementAssetEntry, Factory, AssetData ).HighlightText( HightlightText );
+		CustomContent->SetVisibility(EVisibility::Collapsed);
+		ListView->SetVisibility(EVisibility::Visible);
+		ListView->RequestListRefresh();
+	}
 }
 
-
-TSharedRef< SPlacementAssetEntry > BuildDraggableAssetWidget( UClass* InAssetClass, const FAssetData& InAssetData, FName InClassThumbnailBrushOverride = NAME_None, TOptional<FLinearColor> InAssetTypeColorOverride = TOptional<FLinearColor>(), TAttribute<FText> HightlightText = TAttribute<FText>() ) 
+bool SPlacementModeTools::IsSearchActive() const
 {
-	UActorFactory* Factory = GEditor->FindActorFactoryByClass( InAssetClass );
-	return SNew( SPlacementAssetEntry, Factory, InAssetData )
-		.ClassThumbnailBrushOverride( InClassThumbnailBrushOverride )
-		.AlwaysUseGenericThumbnail(true)
-		.AssetTypeColorOverride( InAssetTypeColorOverride )
-		.HighlightText( HightlightText );
+	return !SearchTextFilter->GetRawFilterText().IsEmpty();
 }
 
-TSharedRef< SWidget > SPlacementModeTools::BuildLightsWidget()
+ECheckBoxState SPlacementModeTools::GetPlacementTabCheckedState( FName CategoryName ) const
 {
-	return SNew( SVerticalBox )
-
-	// Lights
-	+ SVerticalBox::Slot()
-	.AutoHeight()
-	[
-		BuildDraggableAssetWidget( UActorFactoryDirectionalLight::StaticClass() )
-	]
-
-	+ SVerticalBox::Slot()
-	.AutoHeight()
-	[
-		BuildDraggableAssetWidget( UActorFactoryPointLight::StaticClass() )
-	]
-	
-	+ SVerticalBox::Slot()
-	.AutoHeight()
-	[
-		BuildDraggableAssetWidget( UActorFactorySpotLight::StaticClass() )
-	]
-	
-	+ SVerticalBox::Slot()
-	.AutoHeight()
-	[
-		BuildDraggableAssetWidget( UActorFactorySkyLight::StaticClass() )
-	];
+	return ActiveTabName == CategoryName ? ECheckBoxState::Checked : ECheckBoxState::Unchecked;
 }
 
-TSharedRef< SWidget > SPlacementModeTools::BuildVisualWidget()
+EVisibility SPlacementModeTools::GetFailedSearchVisibility() const
 {
-	UActorFactory* PPFactory = GEditor->FindActorFactoryByClassForActorClass( UActorFactoryBoxVolume::StaticClass(), APostProcessVolume::StaticClass() );
-
-	return SNew( SVerticalBox )
-
-	// Visual
-	+ SVerticalBox::Slot()
-	.AutoHeight()
-	[
-		SNew( SPlacementAssetEntry, PPFactory, FAssetData( APostProcessVolume::StaticClass() ) )
-	]
-
-	+ SVerticalBox::Slot()
-	.AutoHeight()
-	[
-		BuildDraggableAssetWidget( UActorFactoryAtmosphericFog::StaticClass() )
-	]
-	
-	+ SVerticalBox::Slot()
-	.AutoHeight()
-	[
-		BuildDraggableAssetWidget( UActorFactoryExponentialHeightFog::StaticClass() )
-	]
-	
-	+ SVerticalBox::Slot()
-	.AutoHeight()
-	[
-		BuildDraggableAssetWidget( UActorFactorySphereReflectionCapture::StaticClass() )
-	]
-
-	+ SVerticalBox::Slot()
-	.AutoHeight()
-	[
-		BuildDraggableAssetWidget( UActorFactoryBoxReflectionCapture::StaticClass() )
-	]
-	
-	+ SVerticalBox::Slot()
-	.AutoHeight()
-	[
-		BuildDraggableAssetWidget( UActorFactoryDeferredDecal::StaticClass() )
-	];
+	if (!IsSearchActive() || FilteredItems.Num())
+	{
+		return EVisibility::Collapsed;
+	}
+	return EVisibility::Visible;
 }
 
-static TOptional<FLinearColor> GetBasicShapeColorOverride()
+EVisibility SPlacementModeTools::GetListViewVisibility()const
 {
-	// Get color for basic shapes.  It should appear like all the other basic types
-	static TOptional<FLinearColor> BasicShapeColorOverride;
-	
-	if( !BasicShapeColorOverride.IsSet() )
+	if (IsSearchActive() || FilteredItems.Num())
 	{
-		FAssetToolsModule& AssetToolsModule = FModuleManager::LoadModuleChecked<FAssetToolsModule>(TEXT("AssetTools"));
-		TSharedPtr<IAssetTypeActions> AssetTypeActions;
-		AssetTypeActions = AssetToolsModule.Get().GetAssetTypeActionsForClass(UClass::StaticClass()).Pin();
-		if (AssetTypeActions.IsValid())
-		{
-			BasicShapeColorOverride = TOptional<FLinearColor>(AssetTypeActions->GetTypeColor());
-		}
+		return EVisibility::Visible;
 	}
-	return BasicShapeColorOverride;
+	return EVisibility::Collapsed;
 }
 
-TSharedRef< SWidget > SPlacementModeTools::BuildBasicWidget()
+EVisibility SPlacementModeTools::GetTabsVisibility() const
 {
-	// Basics
-	TSharedRef<SVerticalBox> VerticalBox = SNew( SVerticalBox )
-	+SVerticalBox::Slot()
-	.AutoHeight()
-	[
-		BuildDraggableAssetWidget(UActorFactoryEmptyActor::StaticClass())
-	]
-	+SVerticalBox::Slot()
-	.AutoHeight()
-	[
-		BuildDraggableAssetWidget(UActorFactoryCharacter::StaticClass())
-	]
-	+SVerticalBox::Slot()
-	.AutoHeight()
-	[
-		BuildDraggableAssetWidget(UActorFactoryPawn::StaticClass())
-	]
-	.AutoHeight()
-	[
-		BuildDraggableAssetWidget(UActorFactoryPointLight::StaticClass())
-	]
-	+ SVerticalBox::Slot()
-	.AutoHeight()
-	[
-		BuildDraggableAssetWidget(UActorFactoryPlayerStart::StaticClass())
-	]
-	+ SVerticalBox::Slot()
-	.AutoHeight()
-	[
-		// Cube
-		BuildDraggableAssetWidget(UActorFactoryBasicShape::StaticClass(), FAssetData( LoadObject<UStaticMesh>(nullptr,*UActorFactoryBasicShape::BasicCube.ToString())), FName("ClassThumbnail.Cube"), GetBasicShapeColorOverride() )
-	]
-	+ SVerticalBox::Slot()
-	.AutoHeight()
-	[
-		// Sphere
-		BuildDraggableAssetWidget(UActorFactoryBasicShape::StaticClass(), FAssetData(LoadObject<UStaticMesh>(nullptr,*UActorFactoryBasicShape::BasicSphere.ToString())), FName("ClassThumbnail.Sphere"), GetBasicShapeColorOverride() )
-	]
-	+ SVerticalBox::Slot()
-	.AutoHeight()
-	[
-		// Cylinder
-		BuildDraggableAssetWidget(UActorFactoryBasicShape::StaticClass(), FAssetData(LoadObject<UStaticMesh>(nullptr,*UActorFactoryBasicShape::BasicCylinder.ToString())), FName("ClassThumbnail.Cylinder"), GetBasicShapeColorOverride() )
-	]
-	+ SVerticalBox::Slot()
-	.AutoHeight()
-	[
-		// Cone
-		BuildDraggableAssetWidget(UActorFactoryBasicShape::StaticClass(), FAssetData(LoadObject<UStaticMesh>(nullptr,*UActorFactoryBasicShape::BasicCone.ToString())), FName("ClassThumbnail.Cone"), GetBasicShapeColorOverride() )
-	]
-	+ SVerticalBox::Slot()
-	.AutoHeight()
-	[
-		BuildDraggableAssetWidget( UActorFactoryTriggerBox::StaticClass() )
-	]
-	+ SVerticalBox::Slot()
-	.AutoHeight()
-	[
-		BuildDraggableAssetWidget( UActorFactoryTriggerSphere::StaticClass() )
-	];
-
-	return VerticalBox;
+	return IsSearchActive() ? EVisibility::Collapsed : EVisibility::Visible;
 }
 
-void SPlacementModeTools::UpdateRecentlyPlacedAssets( const TArray< FActorPlacementInfo >& RecentlyPlaced )
+TSharedRef<ITableRow> SPlacementModeTools::OnGenerateWidgetForItem(TSharedPtr<FPlaceableItem> InItem, const TSharedRef<STableViewBase>& OwnerTable)
 {
-	RefreshRecentlyPlaced();
+	return SNew(STableRow<TSharedPtr<FPlaceableItem>>, OwnerTable)
+		[
+			SNew(SPlacementAssetEntry, InItem.ToSharedRef())
+			.HighlightText(this, &SPlacementModeTools::GetHighlightText)
+		];
 }
 
-void SPlacementModeTools::RefreshRecentlyPlaced()
+void SPlacementModeTools::OnPlacementTabChanged( ECheckBoxState NewState, FName CategoryName )
 {
-	RecentlyPlacedContainer->ClearChildren();
-
-	FAssetRegistryModule& AssetRegistryModule = FModuleManager::LoadModuleChecked<FAssetRegistryModule>( TEXT( "AssetRegistry" ) );
-
-	const TArray< FActorPlacementInfo > RecentlyPlaced = IPlacementModeModule::Get().GetRecentlyPlaced();
-	for ( int Index = 0; Index < RecentlyPlaced.Num(); Index++ )
+	if ( NewState == ECheckBoxState::Checked )
 	{
-		UObject* Asset = FindObject<UObject>(NULL, *RecentlyPlaced[Index].ObjectPath);
-		// If asset is pending delete, it will not be marked as RF_Standalone, in which case we skip it
-		if (Asset != nullptr && Asset->HasAnyFlags(RF_Standalone))
-		{
-			FAssetData AssetData = AssetRegistryModule.Get().GetAssetByObjectPath(*RecentlyPlaced[Index].ObjectPath);
-
-			if (AssetData.IsValid())
-			{
-				TArray< FActorFactoryAssetProxy::FMenuItem > AssetMenuOptions;
-				UActorFactory* Factory = FindObject<UActorFactory>(nullptr, *RecentlyPlaced[Index].Factory);
+		ActiveTabName = CategoryName;
+		IPlacementModeModule::Get().RegenerateItemsForCategory(ActiveTabName);
 
-				RecentlyPlacedContainer->AddSlot()
-					[
-						SNew(SPlacementAssetEntry, Factory, AssetData)
-					];
-			}
-		}
+		bNeedsUpdate = true;
 	}
 }
 
-void SPlacementModeTools::RefreshVolumes()
+const FSlateBrush* SPlacementModeTools::PlacementGroupBorderImage( FName CategoryName ) const
 {
-	bVolumesRefreshRequested = false;
-
-	VolumesContainer->ClearChildren();
-
-	TArray< TSharedRef<SPlacementAssetEntry> > Entries;
-
-	// Make a map of UClasses to ActorFactories that support them
-	const TArray< UActorFactory *>& ActorFactories = GEditor->ActorFactories;
-	TMap<UClass*, UActorFactory*> ActorFactoryMap;
-	for ( int32 FactoryIdx = 0; FactoryIdx < ActorFactories.Num(); ++FactoryIdx )
-	{
-		UActorFactory* ActorFactory = ActorFactories[FactoryIdx];
-
-		if ( ActorFactory )
-		{
-			ActorFactoryMap.Add( ActorFactory->GetDefaultActorClass( FAssetData() ), ActorFactory );
-		}
-	}
-
-	// Add loaded classes
-	FText UnusedErrorMessage;
-	FAssetData NoAssetData;
-	for ( TObjectIterator<UClass> ClassIt; ClassIt; ++ClassIt )
+	if ( ActiveTabName == CategoryName )
 	{
-		const UClass* Class = *ClassIt;
-
-		if (!Class->HasAllClassFlags(CLASS_NotPlaceable) &&
-			!Class->HasAnyClassFlags(CLASS_Abstract | CLASS_Deprecated | CLASS_NewerVersionExists) &&
-			Class->IsChildOf(AVolume::StaticClass()) &&
-			Class->ClassGeneratedBy == nullptr )
-		{
-			UActorFactory* Factory = GEditor->FindActorFactoryByClassForActorClass(UActorFactoryBoxVolume::StaticClass(), Class);
-			Entries.Add(SNew(SPlacementAssetEntry, Factory, FAssetData(Class)));
-		}
+		static FName PlacementBrowserActiveTabBarBrush( "PlacementBrowser.ActiveTabBar" );
+		return FEditorStyle::GetBrush( PlacementBrowserActiveTabBarBrush );
 	}
-
-	struct FCompareFactoryByDisplayName
-	{
-		FORCEINLINE bool operator()( const TSharedRef<SPlacementAssetEntry>& A, const TSharedRef<SPlacementAssetEntry>& B ) const
-		{
-			return A->AssetDisplayName.CompareTo( B->AssetDisplayName ) < 0;
-		}
-	};
-
-	Entries.Sort( FCompareFactoryByDisplayName() );
-
-	for ( int32 i = 0; i < Entries.Num(); i++ )
+	else
 	{
-		VolumesContainer->AddSlot()
-		[
-			Entries[i]
-		];
+		return nullptr;
 	}
 }
 
-void SPlacementModeTools::RebuildPlaceableClassWidgetCache()
+void SPlacementModeTools::UpdateRecentlyPlacedAssets( const TArray< FActorPlacementInfo >& RecentlyPlaced )
 {
-	// Make a map of UClasses to ActorFactories that support them
-	const TArray< UActorFactory *>& ActorFactories = GEditor->ActorFactories;
-	TMap<UClass*, UActorFactory*> ActorFactoryMap;
-	for ( int32 FactoryIdx = 0; FactoryIdx < ActorFactories.Num(); ++FactoryIdx )
-	{
-		UActorFactory* ActorFactory = ActorFactories[FactoryIdx];
-
-		if ( ActorFactory )
-		{
-			ActorFactoryMap.Add(ActorFactory->GetDefaultActorClass(FAssetData()), ActorFactory);
-		}
-	}
-
-	TArray< TSharedRef<SPlacementAssetEntry> > Entries;
-
-	TAttribute<FText> HighlightTextAttrib( this, &SPlacementModeTools::GetHighlightText );
-
-	// Add loaded classes
-	FText UnusedErrorMessage;
-	FAssetData NoAssetData;
-	for ( TObjectIterator<UClass> ClassIt; ClassIt; ++ClassIt )
-	{
-		// Don't offer skeleton classes
-		bool bIsSkeletonClass = FKismetEditorUtilities::IsClassABlueprintSkeleton(*ClassIt);
-
-		if ( !ClassIt->HasAllClassFlags(CLASS_NotPlaceable) &&
-			!ClassIt->HasAnyClassFlags(CLASS_Abstract | CLASS_Deprecated | CLASS_NewerVersionExists) &&
-			ClassIt->IsChildOf(AActor::StaticClass()) &&
-			( !ClassIt->IsChildOf(ABrush::StaticClass()) || ClassIt->IsChildOf(AVolume::StaticClass()) ) &&
-			!bIsSkeletonClass )
-		{
-			UActorFactory** ActorFactory = ActorFactoryMap.Find(*ClassIt);
-
-			const bool IsVolume = ClassIt->IsChildOf(AVolume::StaticClass());
-
-			TSharedPtr<SPlacementAssetEntry> Entry;
-
-			if ( IsVolume )
-			{
-				UActorFactory* Factory = GEditor->FindActorFactoryByClassForActorClass(UActorFactoryBoxVolume::StaticClass(), *ClassIt);
-				Entry = SNew(SPlacementAssetEntry, Factory, FAssetData(*ClassIt))
-					.HighlightText(HighlightTextAttrib);
-			}
-			else
-			{
-				if ( !ActorFactory || ( *ActorFactory )->CanCreateActorFrom(NoAssetData, UnusedErrorMessage) )
-				{
-					if ( !ActorFactory )
-					{
-						Entry = SNew(SPlacementAssetEntry, nullptr, FAssetData(*ClassIt))
-							.HighlightText(HighlightTextAttrib);
-					}
-					else
-					{
-						Entry = SNew(SPlacementAssetEntry, *ActorFactory, FAssetData(*ClassIt))
-							.HighlightText(HighlightTextAttrib);
-					}
-				}
-			}
-
-			if ( Entry.IsValid() )
-			{
-				Entries.Add(Entry.ToSharedRef());
-			}
-		}
-	}
-
-
-	// Add empty types and shapes which have special actor factory requirements
-	Entries.Add(BuildDraggableAssetWidget(UActorFactoryEmptyActor::StaticClass(),HighlightTextAttrib));
-	Entries.Add(BuildDraggableAssetWidget(UActorFactoryCharacter::StaticClass(),HighlightTextAttrib));
-	Entries.Add(BuildDraggableAssetWidget(UActorFactoryPawn::StaticClass(),HighlightTextAttrib));
-	Entries.Add(BuildDraggableAssetWidget(UActorFactoryBasicShape::StaticClass(), FAssetData(LoadObject<UStaticMesh>(nullptr, *UActorFactoryBasicShape::BasicCube.ToString())), FName("ClassThumbnail.Cube"), GetBasicShapeColorOverride(),HighlightTextAttrib));
-	Entries.Add(BuildDraggableAssetWidget(UActorFactoryBasicShape::StaticClass(), FAssetData(LoadObject<UStaticMesh>(nullptr, *UActorFactoryBasicShape::BasicSphere.ToString())), FName("ClassThumbnail.Sphere"), GetBasicShapeColorOverride(),HighlightTextAttrib));
-	Entries.Add(BuildDraggableAssetWidget(UActorFactoryBasicShape::StaticClass(), FAssetData(LoadObject<UStaticMesh>(nullptr, *UActorFactoryBasicShape::BasicCylinder.ToString())), FName("ClassThumbnail.Cylinder"), GetBasicShapeColorOverride(),HighlightTextAttrib));
-	Entries.Add(BuildDraggableAssetWidget(UActorFactoryBasicShape::StaticClass(), FAssetData(LoadObject<UStaticMesh>(nullptr, *UActorFactoryBasicShape::BasicCone.ToString())), FName("ClassThumbnail.Cone"), GetBasicShapeColorOverride(),HighlightTextAttrib));
-
-	struct FCompareFactoryByDisplayName
-	{
-		FORCEINLINE bool operator()(const TSharedRef<SPlacementAssetEntry>& A, const TSharedRef<SPlacementAssetEntry>& B) const
-		{
-			return A->AssetDisplayName.CompareTo(B->AssetDisplayName) < 0;
-		}
-	};
-
-	Entries.Sort(FCompareFactoryByDisplayName());
-
-	// Cache the result
-	PlaceableClassWidgets.Reset();
-	for ( int32 i = 0; i < Entries.Num(); i++ )
+	if (GetActiveTab() == FBuiltInPlacementCategories::RecentlyPlaced())
 	{
-		PlaceableClassWidgets.Add(Entries[i]);
+		bRecentlyPlacedRefreshRequested = true;
 	}
 }
 
-void SPlacementModeTools::RefreshPlaceables()
+void SPlacementModeTools::Tick( const FGeometry& AllottedGeometry, const double InCurrentTime, const float InDeltaTime )
 {
 	if ( bPlaceablesFullRefreshRequested )
 	{
-		RebuildPlaceableClassWidgetCache();
-	}
-
-	PlaceablesContainer->ClearChildren();
-	SearchResultsContainer->ClearChildren();
-
-	if ( SearchTextFilter->GetRawFilterText().IsEmpty() )
-	{
-		// Just build the full list with no filtering
-		for ( int32 i = 0; i < PlaceableClassWidgets.Num(); i++ )
-		{
-			PlaceablesContainer->AddSlot()
-			[
-				PlaceableClassWidgets[i]
-			];
-		}
-	}
-	else
-	{
-		// Filter out the widgets that don't belong
-		for ( int32 i = 0; i < PlaceableClassWidgets.Num(); i++ )
-		{
-			if ( SearchTextFilter->PassesFilter( PlaceableClassWidgets[i] ) )
-			{
-				SearchResultsContainer->AddSlot()
-				[
-					PlaceableClassWidgets[i]
-				];
-			}
-		}
-
-		if ( SearchResultsContainer->GetChildren()->Num() == 0 )
-		{
-			SearchResultsContainer->AddSlot()
-				.HAlign(HAlign_Center)
-				.VAlign(VAlign_Fill)
-				[
-					SNew(STextBlock)
-					.Text(NSLOCTEXT("PlacementMode", "NoResultsFound", "No Results Found"))
-				];
-		}
+		IPlacementModeModule::Get().RegenerateItemsForCategory(FBuiltInPlacementCategories::AllClasses());
+		bPlaceablesFullRefreshRequested = false;
+		bNeedsUpdate = true;
 	}
-}
 
-void SPlacementModeTools::Tick( const FGeometry& AllottedGeometry, const double InCurrentTime, const float InDeltaTime )
-{
-	if ( bPlaceablesRefreshRequested || bPlaceablesFullRefreshRequested )
+	if ( bRecentlyPlacedRefreshRequested )
 	{
-		RefreshPlaceables();
-		bPlaceablesRefreshRequested = false;
-		bPlaceablesFullRefreshRequested = false;
+		IPlacementModeModule::Get().RegenerateItemsForCategory(FBuiltInPlacementCategories::RecentlyPlaced());
+		bRecentlyPlacedRefreshRequested = false;
+		bNeedsUpdate = true;
 	}
 
-	if ( bVolumesRefreshRequested )
+	if ( bNeedsUpdate )
 	{
-		RefreshVolumes();
+		UpdateFilteredItems();
 	}
 }
 
@@ -1036,22 +590,17 @@ FReply SPlacementModeTools::OnKeyDown( const FGeometry& MyGeometry, const FKeyEv
 	return Reply;
 }
 
-int32 SPlacementModeTools::GetSelectedPanel() const
-{
-	return SearchTextFilter->GetRawFilterText().IsEmpty() ? 0 : 1;
-}
-
 void SPlacementModeTools::OnSearchChanged(const FText& InFilterText)
 {
 	// If the search text was previously empty we do a full rebuild of our cached widgets
 	// for the placeable widgets.
-	if ( SearchTextFilter->GetRawFilterText().IsEmpty() )
+	if ( !IsSearchActive() )
 	{
 		bPlaceablesFullRefreshRequested = true;
 	}
 	else
 	{
-		bPlaceablesRefreshRequested = true;
+		bNeedsUpdate = true;
 	}
 
 	SearchTextFilter->SetRawFilterText( InFilterText );
diff --git a/Engine/Source/Editor/PlacementMode/Private/SPlacementModeTools.h b/Engine/Source/Editor/PlacementMode/Private/SPlacementModeTools.h
--- a/Engine/Source/Editor/PlacementMode/Private/SPlacementModeTools.h
+++ b/Engine/Source/Editor/PlacementMode/Private/SPlacementModeTools.h
@@ -4,6 +4,7 @@
 #include "ClassIconFinder.h"
 #include "AssetData.h"
 #include "Editor/ContentBrowser/Public/ContentBrowserModule.h"
+#include "IPlacementModeModule.h"
 
 /**
  * A tile representation of the class or the asset.  These are embedded into the views inside
@@ -12,25 +13,14 @@
 class SPlacementAssetEntry : public SCompoundWidget
 {
 public:
+	SLATE_BEGIN_ARGS(SPlacementAssetEntry){}
 
-	SLATE_BEGIN_ARGS(SPlacementAssetEntry)
-		: _LabelOverride()
-		, _AlwaysUseGenericThumbnail(false)
-	{}
-	
-	/** Highlight this text in the text block */
-	SLATE_ATTRIBUTE(FText, HighlightText)
+		/** Highlight this text in the text block */
+		SLATE_ATTRIBUTE(FText, HighlightText)
 
-	SLATE_ARGUMENT(FText, LabelOverride)
-
-	SLATE_ARGUMENT(FName, ClassThumbnailBrushOverride)
-
-	SLATE_ARGUMENT(bool, AlwaysUseGenericThumbnail)
-
-	SLATE_ARGUMENT(TOptional<FLinearColor>, AssetTypeColorOverride)
 	SLATE_END_ARGS()
 
-	void Construct(const FArguments& InArgs, UActorFactory* InFactory, const FAssetData& InAsset);
+	void Construct(const FArguments& InArgs, const TSharedPtr<FPlaceableItem>& InItem);
 
 	virtual FReply OnMouseButtonDown(const FGeometry& MyGeometry, const FPointerEvent& MouseEvent) override;
 	virtual FReply OnMouseButtonUp(const FGeometry& MyGeometry, const FPointerEvent& MouseEvent) override;
@@ -38,10 +28,7 @@ public:
 
 	bool IsPressed() const;
 
-	FText AssetDisplayName;
-
-	UActorFactory* FactoryToUse;
-	FAssetData AssetData;
+	TSharedPtr<FPlaceableItem> Item;
 
 private:
 	const FSlateBrush* GetBorder() const;
@@ -68,89 +55,85 @@ public:
 
 private:
 
-	// Gets the border image for the tab, this is the 'active' orange bar.
-	const FSlateBrush* PlacementGroupBorderImage( int32 PlacementGroupIndex ) const;
-
-	// When the tab is clicked we adjust the check state, so that the right style is displayed.
-	void OnPlacementTabChanged( ECheckBoxState NewState, int32 PlacementGroupIndex );
+	// Begin SWidget
+	virtual void Tick( const FGeometry& AllottedGeometry, const double InCurrentTime, const float InDeltaTime ) override;
+	virtual FReply OnKeyDown( const FGeometry& MyGeometry, const FKeyEvent& InKeyEvent ) override;
+	// End SWidget
 
-	// Gets the tab 'active' state, so that we can show the active style
-	ECheckBoxState GetPlacementTabCheckedState( int32 PlacementGroupIndex ) const;
+private:
 
-	// Create the standard panel displayed when no search is being performed.
-	TSharedRef< SWidget > CreateStandardPanel();
+	/** Creates a tab widget to show on the left that when clicked sets the currently active tab */
+	TSharedRef< SWidget > CreatePlacementGroupTab( const FPlacementCategoryInfo& Info );
 
-	// Creates a tab widget to show on the left that when clicked actives the right widget switcher index.
-	TSharedRef< SWidget > CreatePlacementGroupTab( int32 TabIndex, FText TabText, bool Important );
+	/** Generates a widget for the specified item */
+	TSharedRef<ITableRow> OnGenerateWidgetForItem(TSharedPtr<FPlaceableItem> InItem, const TSharedRef<STableViewBase>& OwnerTable);
 
-	// Builds the lights collection widget.
-	TSharedRef< SWidget > BuildLightsWidget();
+	/** Get the identifier of the currently active tab */
+	FName GetActiveTab() const;
 
-	// Builds the visual collection widget.
-	TSharedRef< SWidget > BuildVisualWidget();
+	/** Check if a search is active */
+	bool IsSearchActive() const;
 
-	// Builds the 'Basic' collection widget.
-	TSharedRef< SWidget > BuildBasicWidget();
+	/** Update the list of filtered items */
+	void UpdateFilteredItems();
 
-	// Called when the recently placed assets changes.
+	/** Called when the recently placed assets changes. */
 	void UpdateRecentlyPlacedAssets( const TArray< FActorPlacementInfo >& RecentlyPlaced );
 
-	// Rebuilds the recently placed widget
-	void RefreshRecentlyPlaced();
+private:
 
-	// Rebuilds the 'Classes' widget.
-	void RefreshPlaceables();
+	/** Gets the border image for the tab, this is the 'active' orange bar. */
+	const FSlateBrush* PlacementGroupBorderImage( FName CategoryName ) const;
 
-	// Rebuilds the 'Volumes' widget.
-	void RefreshVolumes();
+	/** When the tab is clicked we adjust the check state, so that the right style is displayed. */
+	void OnPlacementTabChanged( ECheckBoxState NewState, FName CategoryName );
 
-	// Begin SWidget
-	virtual void Tick( const FGeometry& AllottedGeometry, const double InCurrentTime, const float InDeltaTime ) override;
-	virtual FReply OnKeyDown( const FGeometry& MyGeometry, const FKeyEvent& InKeyEvent ) override;
-	// End SWidget
+	/** Gets the tab 'active' state, so that we can show the active style */
+	ECheckBoxState GetPlacementTabCheckedState( FName CategoryName ) const;
 
 private:
-	void OnSearchChanged(const FText& InFilterText);
 
-	FText GetHighlightText() const;
+	/** Gets the visibility for the failed search text */
+	EVisibility GetFailedSearchVisibility() const;
 
-	// Get the selected panel (normal or search results)
-	int32 GetSelectedPanel() const;
+	/** Gets the visibility for the list view */
+	EVisibility GetListViewVisibility() const;
 
-	void RebuildPlaceableClassWidgetCache();
+	/** Gets the visibility for tabs */
+	EVisibility GetTabsVisibility() const;
 
-	// Flag set when the placeable widget needs to be rebuilt.
-	bool bPlaceablesRefreshRequested;
+private:
 
-	// Flag set when the placeable widget needs to be rebuilt including the placeable widget cache.
+	/** Called when the search text changes */
+	void OnSearchChanged(const FText& InFilterText);
+
+	/** Get the text that should be highlighted on any items */
+	FText GetHighlightText() const;
+
+private:
+	/** Flags to invoke updates of particular caregories on tick */
+	bool bRecentlyPlacedRefreshRequested;
 	bool bPlaceablesFullRefreshRequested;
 
-	// Flag set when the volumes widget needs to be rebuilt.
-	bool bVolumesRefreshRequested;
+	/** Flag to indicate that we need to update the filtered items array */
+	bool bNeedsUpdate;
 
 	// The text filter used to filter the classes
-	typedef TTextFilter<TSharedRef<SPlacementAssetEntry>> FPlacementAssetEntryTextFilter;
+	typedef TTextFilter<const FPlaceableItem&> FPlacementAssetEntryTextFilter;
 	TSharedPtr<FPlacementAssetEntryTextFilter> SearchTextFilter;
 
-	// The search box used to update the filter text
-	TSharedPtr<SSearchBox> SearchBoxPtr;
-
-	// The list of placeable classes in widget form, these get put in the search results or
-	// the all classes view depending on which is currently visible.
-	TArray< TSharedRef<SPlacementAssetEntry> > PlaceableClassWidgets;
-
-	// Holds the widget switcher.
-	TSharedPtr<SWidgetSwitcher> WidgetSwitcher;
+	/** Custom content slot, where a category has a custom generator */
+	TSharedPtr<SBox> CustomContent;
 
-	// The container widget for the recently placed.
-	TSharedPtr< SVerticalBox > RecentlyPlacedContainer;
+	/* The search box used to update the filter text */
+	TSharedPtr<SSearchBox> SearchBoxPtr;
 
-	// The container widget for the volumes.
-	TSharedPtr< SVerticalBox > VolumesContainer;
+	/** Array of filtered items to show in the list view */
+	TArray<TSharedPtr<FPlaceableItem>> FilteredItems;
 
-	// The container widget for the classes.
-	TSharedPtr< SVerticalBox > PlaceablesContainer;
+	/** The name of the currently active tab (where no search is active) */
+	FName ActiveTabName;
 
-	// The container widget for the classes search results.
-	TSharedPtr< SVerticalBox > SearchResultsContainer;
+	/** List view that shows placeable items */
+	TSharedPtr<SListView<TSharedPtr<FPlaceableItem>>> ListView;
 };
diff --git a/Engine/Source/Editor/PlacementMode/Public/IPlacementModeModule.h b/Engine/Source/Editor/PlacementMode/Public/IPlacementModeModule.h
--- a/Engine/Source/Editor/PlacementMode/Public/IPlacementModeModule.h
+++ b/Engine/Source/Editor/PlacementMode/Public/IPlacementModeModule.h
@@ -5,6 +5,151 @@
 #include "ActorPlacementInfo.h"
 #include "IPlacementMode.h"
 
+/**
+ * Struct that defines an identifier for a particular placeable item in this module.
+ * Only be obtainable through IPlacementModeModule::RegisterPlaceableItem
+ */
+class FPlacementModeID
+{
+public:
+	FPlacementModeID(const FPlacementModeID&) = default;
+	FPlacementModeID& operator=(const FPlacementModeID&) = default;
+
+private:
+	friend class FPlacementModeModule;
+
+	FPlacementModeID() {}
+
+	/** The category this item is held within */
+	FName Category;
+
+	/** Unique identifier (always universally unique across categories unless there's been an int overflow) */
+	uint32 UniqueID;
+};
+
+/**
+ * Struct providing information for a user category of placement objects
+ */
+struct FPlacementCategoryInfo
+{
+	FPlacementCategoryInfo(FText InDisplayName, FName InHandle, FString InTag, int32 InSortOrder = 0, bool bInSortable = true)
+		: DisplayName(InDisplayName), UniqueHandle(InHandle), SortOrder(InSortOrder), TagMetaData(MoveTemp(InTag)), bSortable(bInSortable)
+	{
+	}
+
+	/** This category's display name */
+	FText DisplayName;
+
+	/** A unique name for this category */
+	FName UniqueHandle;
+
+	/** Sort order for the category tab (lowest first) */
+	int32 SortOrder;
+
+	/** Optional tag meta data for the tab widget */
+	FString TagMetaData;
+
+	/** Optional generator function used to construct this category's tab content. Called when the tab is activated. */
+	TFunction<TSharedRef<SWidget>()> CustomGenerator;
+
+	/** Whether the items in this category are automatically sortable by name. False if the items are already sorted. */
+	bool bSortable;
+};
+
+/**
+ * Structure defining a placeable item in the placement mode panel
+ */
+struct FPlaceableItem
+{
+	/** Default constructor */
+	FPlaceableItem()
+		: Factory(nullptr)
+	{}
+
+	/** Constructor that takes aspecific factory and asset */
+	FPlaceableItem(UActorFactory* InFactory, const FAssetData& InAssetData, TOptional<int32> InSortOrder = TOptional<int32>())
+		: Factory(InFactory)
+		, AssetData(InAssetData)
+		, SortOrder(InSortOrder)
+	{
+		bAlwaysUseGenericThumbnail = false;
+		AutoSetDisplayName();
+	}
+
+	/** Constructor for any placeable class */
+	FPlaceableItem(UClass& InAssetClass, TOptional<int32> InSortOrder = TOptional<int32>())
+		: Factory(GEditor->FindActorFactoryByClass(&InAssetClass))
+		, AssetData(Factory ? Factory->GetDefaultActorClass( FAssetData() ) : FAssetData())
+		, SortOrder(InSortOrder)
+	{
+		bAlwaysUseGenericThumbnail = false;
+		AutoSetDisplayName();
+	}
+
+	/** Constructor for any placeable class with associated asset data and brush overrides */
+	FPlaceableItem(
+		UClass& InAssetClass,
+		const FAssetData& InAssetData,
+		FName InClassThumbnailBrushOverride = NAME_None,
+		TOptional<FLinearColor> InAssetTypeColorOverride = TOptional<FLinearColor>(),
+		TOptional<int32> InSortOrder = TOptional<int32>()
+	)
+		: Factory(GEditor->FindActorFactoryByClass(&InAssetClass))
+		, AssetData(InAssetData)
+		, ClassThumbnailBrushOverride(InClassThumbnailBrushOverride)
+		, AssetTypeColorOverride(InAssetTypeColorOverride)
+		, SortOrder(InSortOrder)
+	{
+		bAlwaysUseGenericThumbnail = true;
+		AutoSetDisplayName();
+	}
+
+	/** Automatically set this item's display name from its class or asset */
+	void AutoSetDisplayName()
+	{
+		if (AssetData.GetClass() == UClass::StaticClass())
+		{
+			DisplayName = FText::FromString( FName::NameToDisplayString(AssetData.AssetName.ToString(), false));
+		}
+		else
+		{
+			DisplayName = FText::FromName(AssetData.AssetName);
+		}
+	}
+
+	/** The factory used to create an instance of this placeable item */
+	UActorFactory* Factory;
+
+	/** Asset data pertaining to the class */
+	FAssetData AssetData;
+
+	/** This item's display name */
+	FText DisplayName;
+
+	/** Optional override for the thumbnail brush (passed to FClassIconFinder::FindThumbnailForClass in the form ClassThumbnail.<override>) */
+	FName ClassThumbnailBrushOverride;
+
+	/** Whether to always use the generic thumbnail for this item or not */
+	bool bAlwaysUseGenericThumbnail;
+
+	/** Optional overridden color tint for the asset */
+	TOptional<FLinearColor> AssetTypeColorOverride;
+
+	/** Optional sort order (lowest first). Overrides default class name sorting. */
+	TOptional<int32> SortOrder;
+};
+
+/** Structure of built-in placement categories. Defined as functions to enable external use without linkage */
+struct FBuiltInPlacementCategories
+{
+	static FName RecentlyPlaced()	{ static FName Name("RecentlyPlaced");	return Name; }
+	static FName Basic()			{ static FName Name("Basic");			return Name; }
+	static FName Lights()			{ static FName Name("Lights");			return Name; }
+	static FName Visual()			{ static FName Name("Visual");			return Name; }
+	static FName Volumes()			{ static FName Name("Volumes");			return Name; }
+	static FName AllClasses()		{ static FName Name("AllClasses");		return Name; }
+};
+
 class IPlacementModeModule : public IModuleInterface
 {
 public:
@@ -64,5 +209,77 @@ public:
 	DECLARE_EVENT_OneParam( IPlacementMode, FOnStoppedPlacingEvent, bool /*bWasSuccessfullyPlaced*/ );
 	virtual FOnStoppedPlacingEvent& OnStoppedPlacing() = 0;
 	virtual void BroadcastStoppedPlacing( bool bWasSuccessfullyPlaced ) = 0;
+
+public:
+
+	/**
+	 * Register a new category of placement items
+	 *
+	 * @param Info		Information pertaining to the category
+	 * @return true on success, false on failure (probably if the category's unique handle is already in use)
+	 */
+	virtual bool RegisterPlacementCategory(const FPlacementCategoryInfo& Info) = 0;
+
+	/**
+	 * Retrieve an already registered category
+	 *
+	 * @param UniqueHandle	The unique handle of the category to retrieve
+	 * @return Ptr to the category's information structure, or nullptr if it does not exit
+	 */
+	virtual const FPlacementCategoryInfo* GetRegisteredPlacementCategory(FName UniqueHandle) const = 0;
+
+	/**
+	 * Populate the specified array with all registered category information, sorted by SortOrder
+	 *
+	 * @param OutCategories	The array to populate with registered category inforamtion
+	 */
+	virtual void GetSortedCategories(TArray<FPlacementCategoryInfo>& OutCategories) const = 0;
+
+	/**
+	 * Unregister a previously registered category
+	 *
+	 * @param UniqueHandle	The unique handle of the category to unregister
+	 */
+	virtual void UnregisterPlacementCategory(FName Handle) = 0;
+
+	/**
+	 * Register a new placeable item for the specified category
+	 *
+	 * @param Item			The placeable item to register
+	 * @param CategoryName	Unique handle to the category in which to place this item
+	 * @return Optional unique identifier for the registered item, or empty on failure (if the category doesn't exist)
+	 */
+	virtual TOptional<FPlacementModeID> RegisterPlaceableItem(FName CategoryName, const TSharedRef<FPlaceableItem>& Item) = 0;
+
+	/**
+	 * Unregister a previously registered placeable item
+	 *
+	 * @param ID			Unique identifier for the item. Will have been obtained from a previous call to RegisterPlaceableItem
+	 */
+	virtual void UnregisterPlaceableItem(FPlacementModeID ID) = 0;
+
+	/**
+	 * Get all the items in a given category, unsorted
+	 *
+	 * @param Category		The unique handle of the category to get items for
+	 * @param OutItems		Array to populate with the items in this category
+	 */
+	virtual void GetItemsForCategory(FName Category, TArray<TSharedPtr<FPlaceableItem>>& OutItems) const = 0;
+
+	/**
+	 * Get all the items in a given category, filtered by the specified predicate
+	 *
+	 * @param Category		The unique handle of the category to get items for
+	 * @param OutItems		Array to populate with the items in this category
+	 * @param Filter 		Filter predicate used to filter out items. Return true to pass the filter, false otherwise
+	 */
+	virtual void GetFilteredItemsForCategory(FName Category, TArray<TSharedPtr<FPlaceableItem>>& OutItems, TFunctionRef<bool(const TSharedPtr<FPlaceableItem>&)> Filter) const = 0;
+
+	/**
+	 * Instruct the category associated with the specified unique handle that it should regenerate its items
+	 *
+	 * @param Category		The unique handle of the category to get items for
+	 */
+	virtual void RegenerateItemsForCategory(FName Category) = 0;
 };
 
diff --git a/Engine/Source/Editor/PropertyEditor/Private/UserInterface/PropertyEditor/SPropertyEditorNumeric.h b/Engine/Source/Editor/PropertyEditor/Private/UserInterface/PropertyEditor/SPropertyEditorNumeric.h
--- a/Engine/Source/Editor/PropertyEditor/Private/UserInterface/PropertyEditor/SPropertyEditorNumeric.h
+++ b/Engine/Source/Editor/PropertyEditor/Private/UserInterface/PropertyEditor/SPropertyEditorNumeric.h
@@ -100,21 +100,37 @@ public:
 			&& !PropertyNode->GetProperty()->GetBoolMetaData("NoSpinbox");
 
 		// Set up the correct type interface if we want to display units on the property editor
-		if (FUnitConversion::Settings().ShouldDisplayUnits())
+
+		// First off, check for ForceUnits= meta data. This meta tag tells us to interpret, and always display the value in these units. FUnitConversion::Settings().ShouldDisplayUnits does not apply to suce properties
+		const FString& ForcedUnits = InPropertyEditor->GetProperty()->GetMetaData(TEXT("ForceUnits"));
+		auto PropertyUnits = FUnitConversion::UnitFromString(*ForcedUnits);
+		if (PropertyUnits.IsSet())
+		{
+			// Create the type interface and set up the default input units if they are compatible
+			TypeInterface = MakeShareable(new TNumericUnitTypeInterface<NumericType>(PropertyUnits.GetValue()));
+			TypeInterface->FixedDisplayUnits = PropertyUnits.GetValue();
+		}
+		// If that's not set, we fall back to Units=xxx which calculates the most appropriate unit to display in
+		else
 		{
-			const FString& Units = InPropertyEditor->GetProperty()->GetMetaData(TEXT("Units"));
-			auto PropertyUnits = FUnitConversion::UnitFromString(*Units);
+			if (FUnitConversion::Settings().ShouldDisplayUnits())
+			{
+				const FString& DynamicUnits = InPropertyEditor->GetProperty()->GetMetaData(TEXT("Units"));
+				PropertyUnits = FUnitConversion::UnitFromString(*DynamicUnits);
+			}
 
-			if (PropertyUnits.IsSet())
+			if (!PropertyUnits.IsSet())
 			{
-				// Create the type interface and set up the default input units if they are compatible
-				TypeInterface = MakeShareable(new TNumericUnitTypeInterface<NumericType>(PropertyUnits.GetValue()));
+				PropertyUnits = EUnit::Unspecified;
+			}
 
-				auto Value = OnGetValue();
-				if (Value.IsSet())
-				{
-					TypeInterface->SetupFixedDisplay(Value.GetValue());
-				}
+			// Create the type interface and set up the default input units if they are compatible
+			TypeInterface = MakeShareable(new TNumericUnitTypeInterface<NumericType>(PropertyUnits.GetValue()));
+			auto Value = OnGetValue();
+
+			if (Value.IsSet())
+			{
+				TypeInterface->SetupFixedDisplay(Value.GetValue());
 			}
 		}
 
@@ -134,6 +150,7 @@ public:
 				.UndeterminedString(LOCTEXT("MultipleValues", "Multiple Values"))
 				.OnValueChanged(this, &SPropertyEditorNumeric<NumericType>::OnValueChanged)
 				.OnValueCommitted(this, &SPropertyEditorNumeric<NumericType>::OnValueCommitted)
+				.OnUndeterminedValueCommitted(this, &SPropertyEditorNumeric<NumericType>::OnUndeterminedValueCommitted)
 				.OnBeginSliderMovement(this, &SPropertyEditorNumeric<NumericType>::OnBeginSliderMovement)
 				.OnEndSliderMovement(this, &SPropertyEditorNumeric<NumericType>::OnEndSliderMovement)
 				.TypeInterface(TypeInterface)
@@ -299,6 +316,30 @@ private:
 			TypeInterface->SetupFixedDisplay(NewValue);
 		}
 	}
+
+	void OnUndeterminedValueCommitted( FText NewValue, ETextCommit::Type CommitType )
+	{
+		const TSharedRef<IPropertyHandle> PropertyHandle = PropertyEditor->GetPropertyHandle();
+		const FString NewValueString = NewValue.ToString();
+		TArray<FString> PerObjectValues;
+		
+		// evaluate expression for each property value
+		PropertyHandle->GetPerObjectValues(PerObjectValues);
+
+		for (auto& Value : PerObjectValues)
+		{
+			NumericType OldNumericValue;
+			TTypeFromString<NumericType>::FromString(OldNumericValue, *Value);
+			TOptional<NumericType> NewNumericValue = TypeInterface->FromString(NewValueString, OldNumericValue);
+
+			if (NewNumericValue.IsSet())
+			{
+				Value = TTypeToString<NumericType>::ToString(NewNumericValue.GetValue());
+			}
+		}
+
+		PropertyHandle->SetPerObjectValues(PerObjectValues);
+	}
 	
 	/**
 	 * Called when the slider begins to move.  We create a transaction here to undo the property
diff --git a/Engine/Source/Editor/Sequencer/Private/DisplayNodes/SequencerDisplayNode.cpp b/Engine/Source/Editor/Sequencer/Private/DisplayNodes/SequencerDisplayNode.cpp
--- a/Engine/Source/Editor/Sequencer/Private/DisplayNodes/SequencerDisplayNode.cpp
+++ b/Engine/Source/Editor/Sequencer/Private/DisplayNodes/SequencerDisplayNode.cpp
@@ -141,11 +141,14 @@ void SSequencerObjectTrack::AddKeyTime(float NewTime, TArray<float>& OutKeyTimes
 
 
 FSequencerDisplayNode::FSequencerDisplayNode( FName InNodeName, TSharedPtr<FSequencerDisplayNode> InParentNode, FSequencerNodeTree& InParentTree )
-	: ParentNode( InParentNode )
+	: VirtualTop( 0.f )
+	, VirtualBottom( 0.f )
+	, ParentNode( InParentNode )
 	, ParentTree( InParentTree )
 	, NodeName( InNodeName )
 	, bExpanded( false )
-{ }
+{
+}
 
 
 void FSequencerDisplayNode::Initialize(float InVirtualTop, float InVirtualBottom)
@@ -242,7 +245,6 @@ bool FSequencerDisplayNode::TraverseVisible_ParentFirst(const TFunctionRef<bool(
 	return true;
 }
 
-
 TSharedRef<FSequencerSectionCategoryNode> FSequencerDisplayNode::AddCategoryNode( FName CategoryName, const FText& DisplayLabel )
 {
 	TSharedPtr<FSequencerSectionCategoryNode> CategoryNode;
@@ -331,18 +333,37 @@ void FSequencerDisplayNode::AddKeyAreaNode(FName KeyAreaName, const FText& Displ
 
 TSharedRef<SWidget> FSequencerDisplayNode::GenerateContainerWidgetForOutliner(const TSharedRef<SSequencerTreeViewRow>& InRow)
 {
-	auto NewWidget = SNew(SAnimationOutlinerTreeNode, SharedThis(this), InRow);
-	TreeNodeWidgetPtr = NewWidget;
+	auto NewWidget = SNew(SAnimationOutlinerTreeNode, SharedThis(this), InRow)
+	.IconBrush(this, &FSequencerDisplayNode::GetIconBrush)
+	.IconOverlayBrush(this, &FSequencerDisplayNode::GetIconOverlayBrush)
+	.IconToolTipText(this, &FSequencerDisplayNode::GetIconToolTipText)
+	.CustomContent()
+	[
+		GetCustomOutlinerContent()
+	];
 
 	return NewWidget;
 }
 
-
-TSharedRef<SWidget> FSequencerDisplayNode::GenerateEditWidgetForOutliner()
+TSharedRef<SWidget> FSequencerDisplayNode::GetCustomOutlinerContent()
 {
 	return SNew(SSpacer);
 }
 
+const FSlateBrush* FSequencerDisplayNode::GetIconBrush() const
+{
+	return nullptr;
+}
+
+const FSlateBrush* FSequencerDisplayNode::GetIconOverlayBrush() const
+{
+	return nullptr;
+}
+
+FText FSequencerDisplayNode::GetIconToolTipText() const
+{
+	return FText();
+}
 
 TSharedRef<SWidget> FSequencerDisplayNode::GenerateWidgetForSectionArea(const TAttribute< TRange<float> >& ViewRange)
 {
@@ -404,7 +425,7 @@ TSharedPtr<SWidget> FSequencerDisplayNode::OnSummonContextMenu(const FGeometry&
 {
 	// @todo sequencer replace with UI Commands instead of faking it
 	const bool bShouldCloseWindowAfterMenuSelection = true;
-	FMenuBuilder MenuBuilder(bShouldCloseWindowAfterMenuSelection, nullptr);
+	FMenuBuilder MenuBuilder(bShouldCloseWindowAfterMenuSelection, GetSequencer().GetCommandBindings());
 
 	BuildContextMenu(MenuBuilder);
 
@@ -456,8 +477,8 @@ void FSequencerDisplayNode::BuildContextMenu(FMenuBuilder& MenuBuilder)
 			LOCTEXT("RenameNodeTooltip", "Rename this track"),
 			FSlateIcon(),
 			FUIAction(
-				FExecuteAction::CreateSP(this, &FSequencerDisplayNode::HandleContextMenuRenameNodeExecute, ThisNode),
-				FCanExecuteAction::CreateSP(this, &FSequencerDisplayNode::HandleContextMenuRenameNodeCanExecute, ThisNode)
+				FExecuteAction::CreateSP(this, &FSequencerDisplayNode::HandleContextMenuRenameNodeExecute),
+				FCanExecuteAction::CreateSP(this, &FSequencerDisplayNode::HandleContextMenuRenameNodeCanExecute)
 			)
 		);
 	}
@@ -502,10 +523,9 @@ bool FSequencerDisplayNode::IsHidden() const
 
 bool FSequencerDisplayNode::IsHovered() const
 {
-	return TreeNodeWidgetPtr.IsValid() && TreeNodeWidgetPtr.Pin()->IsHovered();
+	return ParentTree.GetHoveredNode().Get() == this;
 }
 
-
 TSharedRef<FGroupedKeyArea> FSequencerDisplayNode::GetKeyGrouping(int32 InSectionIndex)
 {
 	if (!KeyGroupings.IsValidIndex(InSectionIndex))
@@ -546,20 +566,15 @@ TSharedRef<FGroupedKeyArea> FSequencerDisplayNode::UpdateKeyGrouping(int32 InSec
 }
 
 
-void FSequencerDisplayNode::HandleContextMenuRenameNodeExecute(TSharedRef<FSequencerDisplayNode> NodeToBeRenamed)
+void FSequencerDisplayNode::HandleContextMenuRenameNodeExecute()
 {
-	TSharedPtr<SAnimationOutlinerTreeNode> TreeNodeWidget = TreeNodeWidgetPtr.Pin();
-
-	if (TreeNodeWidget.IsValid())
-	{
-		TreeNodeWidget->EnterRenameMode();
-	}
+	RenameRequestedEvent.Broadcast();
 }
 
 
-bool FSequencerDisplayNode::HandleContextMenuRenameNodeCanExecute(TSharedRef<FSequencerDisplayNode> NodeToBeRenamed) const
+bool FSequencerDisplayNode::HandleContextMenuRenameNodeCanExecute() const
 {
-	return NodeToBeRenamed->CanRenameNode();
+	return CanRenameNode();
 }
 
 
diff --git a/Engine/Source/Editor/Sequencer/Private/DisplayNodes/SequencerDisplayNode.h b/Engine/Source/Editor/Sequencer/Private/DisplayNodes/SequencerDisplayNode.h
--- a/Engine/Source/Editor/Sequencer/Private/DisplayNodes/SequencerDisplayNode.h
+++ b/Engine/Source/Editor/Sequencer/Private/DisplayNodes/SequencerDisplayNode.h
@@ -17,7 +17,7 @@ struct FNodePadding
 {
 	FNodePadding(float InUniform) : Top(InUniform), Bottom(InUniform) { }
 	FNodePadding(float InTop, float InBottom) : Top(InTop), Bottom(InBottom) { }
-	
+
 	/** @return The sum total of the separate padding values */
 	float Combined() const
 	{
@@ -31,7 +31,6 @@ struct FNodePadding
 	float Bottom;
 };
 
-
 namespace ESequencerNode
 {
 	enum Type
@@ -44,6 +43,8 @@ namespace ESequencerNode
 		KeyArea,
 		/* Displays a category */
 		Category,
+		/* Benign spacer node */
+		Spacer,
 	};
 }
 
@@ -150,11 +151,11 @@ public:
 	virtual TSharedRef<SWidget> GenerateContainerWidgetForOutliner(const TSharedRef<SSequencerTreeViewRow>& InRow);
 
 	/**
-	* Generates a widget editing a row in the animation outliner portion of the track area
-	*
-	* @return Generated outliner edit widget
-	*/
-	virtual TSharedRef<SWidget> GenerateEditWidgetForOutliner();
+	 * Customizes an outliner widget that is to represent this node
+	 * 
+	 * @return Content to display on the outliner node
+	 */
+	virtual TSharedRef<SWidget> GetCustomOutlinerContent();
 
 	/**
 	 * Generates a widget for display in the section area portion of the track area
@@ -164,6 +165,27 @@ public:
 	 */
 	virtual TSharedRef<SWidget> GenerateWidgetForSectionArea( const TAttribute<TRange<float>>& ViewRange );
 
+	/**
+	 * Gets an icon that represents this sequencer display node
+	 * 
+	 * @return This node's representative icon
+	 */
+	virtual const FSlateBrush* GetIconBrush() const;
+
+	/**
+	 * Get a brush to overlay on top of the icon for this node
+	 * 
+	 * @return An overlay brush, or nullptr
+	 */
+	virtual const FSlateBrush* GetIconOverlayBrush() const;
+
+	/**
+	 * Get the tooltip text to display for this node's icon
+	 * 
+	 * @return Text to display on the icon
+	 */
+	virtual FText GetIconToolTipText() const;
+
 	/**
 	 * Get the display node that is ultimately responsible for constructing a section area widget for this node.
 	 * Could return this node itself, or a parent node
@@ -256,12 +278,21 @@ public:
 	}
 
 	/**
-	 * @return The parent of this node                                                              
+	 * @return The parent of this node
 	 */
 	TSharedPtr<FSequencerDisplayNode> GetParent() const
 	{
 		return ParentNode.Pin();
 	}
+
+	/**
+	 * @return The outermost parent of this node
+	 */
+	TSharedRef<FSequencerDisplayNode> GetOutermostParent()
+	{
+		TSharedPtr<FSequencerDisplayNode> Parent = ParentNode.Pin();
+		return Parent.IsValid() ? Parent->GetOutermostParent() : AsShared();
+	}
 	
 	/** Gets the sequencer that owns this node */
 	FSequencer& GetSequencer() const
@@ -320,6 +351,9 @@ public:
 	/** Get the key grouping for the specified section index */
 	TSharedRef<FGroupedKeyArea> GetKeyGrouping(int32 InSectionIndex);
 
+	DECLARE_EVENT(FSequencerDisplayNode, FRequestRenameEvent);
+	FRequestRenameEvent& OnRenameRequested() { return RenameRequestedEvent; }
+
 public:
 
 	/** Temporarily disable dynamic regeneration of key groupings. This prevents overlapping key groups from being amalgamated during drags. Key times will continue to update correctly. */
@@ -331,10 +365,10 @@ public:
 private:
 
 	/** Callback for executing a "Rename Node" context menu action. */
-	void HandleContextMenuRenameNodeExecute(TSharedRef<FSequencerDisplayNode> NodeToBeRenamed);
+	void HandleContextMenuRenameNodeExecute();
 
 	/** Callback for determining whether a "Rename Node" context menu action can execute. */
-	bool HandleContextMenuRenameNodeCanExecute(TSharedRef<FSequencerDisplayNode> NodeToBeRenamed) const;
+	bool HandleContextMenuRenameNodeCanExecute() const;
 
 protected:
 
@@ -364,6 +398,6 @@ protected:
 	/** Transient grouped keys for this node */
 	TArray<TSharedPtr<FGroupedKeyArea>> KeyGroupings;
 
-	/** The outliner tree node widget that belongs to this display node. */
-	TWeakPtr<SAnimationOutlinerTreeNode> TreeNodeWidgetPtr;
+	/** Event that is triggered when rename is requested */
+	FRequestRenameEvent RenameRequestedEvent;
 };
diff --git a/Engine/Source/Editor/Sequencer/Private/DisplayNodes/SequencerObjectBindingNode.cpp b/Engine/Source/Editor/Sequencer/Private/DisplayNodes/SequencerObjectBindingNode.cpp
--- a/Engine/Source/Editor/Sequencer/Private/DisplayNodes/SequencerObjectBindingNode.cpp
+++ b/Engine/Source/Editor/Sequencer/Private/DisplayNodes/SequencerObjectBindingNode.cpp
@@ -8,7 +8,9 @@
 #include "MovieSceneTrack.h"
 #include "ObjectEditorUtils.h"
 #include "AssetEditorManager.h"
-
+#include "SequencerUtilities.h"
+#include "ClassIconFinder.h"
+#include "SequencerCommands.h"
 
 #define LOCTEXT_NAMESPACE "FObjectBindingNode"
 
@@ -60,6 +62,28 @@ struct PropertyMenuData
 };
 
 
+
+FSequencerObjectBindingNode::FSequencerObjectBindingNode(FName NodeName, const FText& InDisplayName, const FGuid& InObjectBinding, TSharedPtr<FSequencerDisplayNode> InParentNode, FSequencerNodeTree& InParentTree)
+	: FSequencerDisplayNode(NodeName, InParentNode, InParentTree)
+	, ObjectBinding(InObjectBinding)
+	, DefaultDisplayName(InDisplayName)
+{
+	UMovieScene* MovieScene = GetSequencer().GetFocusedMovieSceneSequence()->GetMovieScene();
+
+	if (MovieScene->FindPossessable(ObjectBinding))
+	{
+		BindingType = EObjectBindingType::Possessable;
+	}
+	else if (MovieScene->FindSpawnable(ObjectBinding))
+	{
+		BindingType = EObjectBindingType::Spawnable;
+	}
+	else
+	{
+		BindingType = EObjectBindingType::Unknown;
+	}
+}
+
 /* FSequencerDisplayNode interface
  *****************************************************************************/
 
@@ -86,6 +110,13 @@ void FSequencerObjectBindingNode::BuildContextMenu(FMenuBuilder& MenuBuilder)
 					}))
 				);
 			}
+
+			MenuBuilder.AddSubMenu(
+				LOCTEXT("OwnerLabel", "Spawned Object Owner"),
+				LOCTEXT("OwnerTooltip", "Specifies how the spawned object is to be owned"),
+				FNewMenuDelegate::CreateSP(this, &FSequencerObjectBindingNode::AddSpawnOwnershipMenu)
+			);
+
 		}
 		else
 		{
@@ -100,6 +131,8 @@ void FSequencerObjectBindingNode::BuildContextMenu(FMenuBuilder& MenuBuilder)
 					FText::Format( LOCTEXT("AssignActorTooltip", "Assign an actor to this track"), Args ),
 					FNewMenuDelegate::CreateRaw(&GetSequencer(), &FSequencer::AssignActor, ObjectBinding));
 			}
+
+			MenuBuilder.AddMenuEntry( FSequencerCommands::Get().ConvertToSpawnable );
 		}
 
 		MenuBuilder.BeginSection("Organize", LOCTEXT("OrganizeContextMenuSectionName", "Organize"));
@@ -116,55 +149,82 @@ void FSequencerObjectBindingNode::BuildContextMenu(FMenuBuilder& MenuBuilder)
 	FSequencerDisplayNode::BuildContextMenu(MenuBuilder);
 }
 
+void FSequencerObjectBindingNode::AddSpawnOwnershipMenu(FMenuBuilder& MenuBuilder)
+{
+	FMovieSceneSpawnable* Spawnable = GetSequencer().GetFocusedMovieSceneSequence()->GetMovieScene()->FindSpawnable(ObjectBinding);
+	if (!Spawnable)
+	{
+		return;
+	}
+
+	MenuBuilder.AddMenuEntry(
+		LOCTEXT("ThisSequence_Label", "This Sequence"),
+		LOCTEXT("ThisSequence_Tooltip", "Indicates that this sequence will own the spawned object. The object will be destroyed at the end of the sequence."),
+		FSlateIcon(),
+		FUIAction(
+			FExecuteAction::CreateLambda([=]{ Spawnable->SetSpawnOwnership(ESpawnOwnership::InnerSequence); }),
+			FCanExecuteAction(),
+			FIsActionChecked::CreateLambda([=]{ return Spawnable->GetSpawnOwnership() == ESpawnOwnership::InnerSequence; })
+		),
+		NAME_None,
+		EUserInterfaceActionType::ToggleButton
+	);
+
+	MenuBuilder.AddMenuEntry(
+		LOCTEXT("MasterSequence_Label", "Master Sequence"),
+		LOCTEXT("MasterSequence_Tooltip", "Indicates that the outermost sequence will own the spawned object. The object will be destroyed when the outermost sequence stops playing."),
+		FSlateIcon(),
+		FUIAction(
+			FExecuteAction::CreateLambda([=]{ Spawnable->SetSpawnOwnership(ESpawnOwnership::MasterSequence); }),
+			FCanExecuteAction(),
+			FIsActionChecked::CreateLambda([=]{ return Spawnable->GetSpawnOwnership() == ESpawnOwnership::MasterSequence; })
+		),
+		NAME_None,
+		EUserInterfaceActionType::ToggleButton
+	);
+
+	MenuBuilder.AddMenuEntry(
+		LOCTEXT("External_Label", "External"),
+		LOCTEXT("External_Tooltip", "Indicates this object's lifetime is managed externally once spawned. It will not be destroyed by sequencer."),
+		FSlateIcon(),
+		FUIAction(
+			FExecuteAction::CreateLambda([=]{ Spawnable->SetSpawnOwnership(ESpawnOwnership::External); }),
+			FCanExecuteAction(),
+			FIsActionChecked::CreateLambda([=]{ return Spawnable->GetSpawnOwnership() == ESpawnOwnership::External; })
+		),
+		NAME_None,
+		EUserInterfaceActionType::ToggleButton
+	);
+}
 
 bool FSequencerObjectBindingNode::CanRenameNode() const
 {
 	return true;
 }
 
-
-TSharedRef<SWidget> FSequencerObjectBindingNode::GenerateEditWidgetForOutliner()
+TSharedRef<SWidget> FSequencerObjectBindingNode::GetCustomOutlinerContent()
 {
 	// Create a container edit box
-	TSharedPtr<class SHorizontalBox> EditBox;
-	SAssignNew(EditBox, SHorizontalBox);	
+	TSharedRef<SHorizontalBox> BoxPanel = SNew(SHorizontalBox)
+		+ SHorizontalBox::Slot()
+		[
+			SNew(SSpacer)
+		];
+
 
-	// Add the property combo box
-	EditBox.Get()->AddSlot()
+	TAttribute<bool> HoverState = TAttribute<bool>::Create(TAttribute<bool>::FGetter::CreateSP(this, &FSequencerDisplayNode::IsHovered));
+
+	BoxPanel->AddSlot()
+	.AutoWidth()
+	.VAlign(VAlign_Center)
 	[
-		SNew(SComboButton)
-			.ButtonStyle(FEditorStyle::Get(), "FlatButton.Light")
-			.OnGetMenuContent(this, &FSequencerObjectBindingNode::HandleAddTrackComboButtonGetMenuContent)
-			.ContentPadding(FMargin(2, 0))
-			.ButtonContent()
-			[
-				SNew(SHorizontalBox)
-
-				+ SHorizontalBox::Slot()
-					.VAlign(VAlign_Center)
-					.AutoWidth()
-					[
-						SNew(STextBlock)
-							.Font(FEditorStyle::Get().GetFontStyle("FontAwesome.8"))
-							.Text(FText::FromString(FString(TEXT("\xf067"))) /*fa-plus*/)
-					]
-
-				+ SHorizontalBox::Slot()
-					.AutoWidth()
-					.Padding(4, 0, 0, 0)
-					[
-						SNew(STextBlock)
-							.Font(FEditorStyle::GetFontStyle("Sequencer.AnimationOutliner.RegularFont"))
-							.Text(LOCTEXT("AddTrackButton", "Track"))
-					]
-			]
+		FSequencerUtilities::MakeAddButton(LOCTEXT("TrackText", "Track"), FOnGetContent::CreateSP(this, &FSequencerObjectBindingNode::HandleAddTrackComboButtonGetMenuContent), HoverState)
 	];
 
 	const UClass* ObjectClass = GetClassForObjectBinding();
+	GetSequencer().BuildObjectBindingEditButtons(BoxPanel, ObjectBinding, ObjectClass);
 
-	GetSequencer().BuildObjectBindingEditButtons(EditBox, ObjectBinding, ObjectClass);
-
-	return EditBox.ToSharedRef();
+	return BoxPanel;
 }
 
 
@@ -180,6 +240,33 @@ FText FSequencerObjectBindingNode::GetDisplayName() const
 	return DefaultDisplayName;
 }
 
+const FSlateBrush* FSequencerObjectBindingNode::GetIconBrush() const
+{
+	return FClassIconFinder::FindIconForClass(GetClassForObjectBinding());
+}
+
+const FSlateBrush* FSequencerObjectBindingNode::GetIconOverlayBrush() const
+{
+	if (BindingType == EObjectBindingType::Spawnable)
+	{
+		return FEditorStyle::GetBrush("Sequencer.SpawnableIconOverlay");
+	}
+	return nullptr;
+}
+
+FText FSequencerObjectBindingNode::GetIconToolTipText() const
+{
+	if (BindingType == EObjectBindingType::Spawnable)
+	{
+		return LOCTEXT("SpawnableToolTip", "This item is spawned by sequencer according to this object's spawn track.");
+	}
+	else if (BindingType == EObjectBindingType::Possessable)
+	{
+		return LOCTEXT("PossessableToolTip", "This item is a possessable reference to an existing object.");
+	}
+
+	return FText();
+}
 
 float FSequencerObjectBindingNode::GetNodeHeight() const
 {
@@ -189,7 +276,7 @@ float FSequencerObjectBindingNode::GetNodeHeight() const
 
 FNodePadding FSequencerObjectBindingNode::GetNodePadding() const
 {
-	return FNodePadding(SequencerNodeConstants::CommonPadding * 2, 0.f);
+	return FNodePadding(SequencerNodeConstants::CommonPadding);
 }
 
 
@@ -261,10 +348,8 @@ void FSequencerObjectBindingNode::AddPropertyMenuItems(FMenuBuilder& AddTrackMen
 }
 
 
-const UClass* FSequencerObjectBindingNode::GetClassForObjectBinding()
+const UClass* FSequencerObjectBindingNode::GetClassForObjectBinding() const
 {
-	FSequencer& ParentSequencer = GetSequencer();
-
 	UMovieScene* MovieScene = GetSequencer().GetFocusedMovieSceneSequence()->GetMovieScene();
 	FMovieSceneSpawnable* Spawnable = MovieScene->FindSpawnable(ObjectBinding);
 	FMovieScenePossessable* Possessable = MovieScene->FindPossessable(ObjectBinding);
diff --git a/Engine/Source/Editor/Sequencer/Private/DisplayNodes/SequencerObjectBindingNode.h b/Engine/Source/Editor/Sequencer/Private/DisplayNodes/SequencerObjectBindingNode.h
--- a/Engine/Source/Editor/Sequencer/Private/DisplayNodes/SequencerObjectBindingNode.h
+++ b/Engine/Source/Editor/Sequencer/Private/DisplayNodes/SequencerObjectBindingNode.h
@@ -4,6 +4,11 @@
 
 #include "SequencerDisplayNode.h"
 
+/** Enumeration specifying what kind of object binding this is */
+enum class EObjectBindingType
+{
+	Possessable, Spawnable, Unknown
+};
 
 /**
  * A node for displaying an object binding
@@ -22,11 +27,7 @@ public:
 	 * @param InParentNode The parent of this node, or nullptr if this is a root node.
 	 * @param InParentTree The tree this node is in.
 	 */
-	FSequencerObjectBindingNode(FName NodeName, const FText& InDisplayName, const FGuid& InObjectBinding, TSharedPtr<FSequencerDisplayNode> InParentNode, FSequencerNodeTree& InParentTree)
-		: FSequencerDisplayNode(NodeName, InParentNode, InParentTree)
-		, ObjectBinding(InObjectBinding)
-		, DefaultDisplayName(InDisplayName)
-	{ }
+	FSequencerObjectBindingNode(FName NodeName, const FText& InDisplayName, const FGuid& InObjectBinding, TSharedPtr<FSequencerDisplayNode> InParentNode, FSequencerNodeTree& InParentTree);
 
 public:
 
@@ -35,15 +36,24 @@ public:
 	{
 		return ObjectBinding;
 	}
-	
+
+	/** Access the cached object binding type for this display node */
+	EObjectBindingType GetBindingType() const
+	{
+		return BindingType;
+	}
+
 public:
 
 	// FSequencerDisplayNode interface
 
 	virtual void BuildContextMenu(FMenuBuilder& MenuBuilder) override;
 	virtual bool CanRenameNode() const override;
-	virtual TSharedRef<SWidget> GenerateEditWidgetForOutliner() override;
+	virtual TSharedRef<SWidget> GetCustomOutlinerContent() override;
 	virtual FText GetDisplayName() const override;
+	virtual const FSlateBrush* GetIconBrush() const override;
+	virtual const FSlateBrush* GetIconOverlayBrush() const override;
+	virtual FText GetIconToolTipText() const override;
 	virtual float GetNodeHeight() const override;
 	virtual FNodePadding GetNodePadding() const override;
 	virtual ESequencerNode::Type GetType() const override;
@@ -53,8 +63,10 @@ protected:
 
 	void AddPropertyMenuItems(FMenuBuilder& AddTrackMenuBuilder, TArray<TArray<UProperty*>> KeyableProperties, int32 PropertyNameIndexStart = 0, int32 PropertyNameIndexEnd = -1);
 
+	void AddSpawnOwnershipMenu(FMenuBuilder& MenuBuilder);
+
 	/** Get class for object binding */
-	const UClass* GetClassForObjectBinding();
+	const UClass* GetClassForObjectBinding() const;
 
 private:
 
@@ -75,4 +87,7 @@ private:
 
 	/** The default display name of the object which is used if the binding manager doesn't provide one. */
 	FText DefaultDisplayName;
+
+	/** Enumeration specifying what kind of object binding this is */
+	EObjectBindingType BindingType;
 };
diff --git a/Engine/Source/Editor/Sequencer/Private/DisplayNodes/SequencerSectionKeyAreaNode.cpp b/Engine/Source/Editor/Sequencer/Private/DisplayNodes/SequencerSectionKeyAreaNode.cpp
--- a/Engine/Source/Editor/Sequencer/Private/DisplayNodes/SequencerSectionKeyAreaNode.cpp
+++ b/Engine/Source/Editor/Sequencer/Private/DisplayNodes/SequencerSectionKeyAreaNode.cpp
@@ -3,7 +3,7 @@
 #include "SequencerPrivatePCH.h"
 #include "IKeyArea.h"
 #include "Sequencer.h"
-
+#include "SKeyNavigationButtons.h"
 
 /* FSectionKeyAreaNode interface
  *****************************************************************************/
@@ -23,7 +23,7 @@ bool FSequencerSectionKeyAreaNode::CanRenameNode() const
 }
 
 
-TSharedRef<SWidget> FSequencerSectionKeyAreaNode::GenerateEditWidgetForOutliner()
+TSharedRef<SWidget> FSequencerSectionKeyAreaNode::GetCustomOutlinerContent()
 {
 	// @todo sequencer: support multiple sections/key areas?
 	TArray<TSharedRef<IKeyArea>> AllKeyAreas = GetAllKeyAreas();
@@ -32,11 +32,35 @@ TSharedRef<SWidget> FSequencerSectionKeyAreaNode::GenerateEditWidgetForOutliner(
 	{
 		if (AllKeyAreas[0]->CanCreateKeyEditor())
 		{
-			return AllKeyAreas[0]->CreateKeyEditor(&GetSequencer());
+			return SNew(SBox)
+			.HAlign(HAlign_Right)
+			.VAlign(VAlign_Center)
+			[
+				SNew(SHorizontalBox)
+
+				+ SHorizontalBox::Slot()
+				.HAlign(HAlign_Right)
+				.VAlign(VAlign_Center)
+				[
+					SNew(SBox)
+					.WidthOverride(100)
+					.HAlign(HAlign_Left)
+					[
+						AllKeyAreas[0]->CreateKeyEditor(&GetSequencer())
+					]
+				]
+
+				+ SHorizontalBox::Slot()
+				.AutoWidth()
+				.VAlign(VAlign_Center)
+				[
+					SNew(SKeyNavigationButtons, AsShared())
+				]
+			];
 		}
 	}
 
-	return FSequencerDisplayNode::GenerateEditWidgetForOutliner();
+	return FSequencerDisplayNode::GetCustomOutlinerContent();
 }
 
 
diff --git a/Engine/Source/Editor/Sequencer/Private/DisplayNodes/SequencerSectionKeyAreaNode.h b/Engine/Source/Editor/Sequencer/Private/DisplayNodes/SequencerSectionKeyAreaNode.h
--- a/Engine/Source/Editor/Sequencer/Private/DisplayNodes/SequencerSectionKeyAreaNode.h
+++ b/Engine/Source/Editor/Sequencer/Private/DisplayNodes/SequencerSectionKeyAreaNode.h
@@ -75,7 +75,7 @@ public:
 	// FSequencerDisplayNode interface
 
 	virtual bool CanRenameNode() const override;
-	virtual TSharedRef<SWidget> GenerateEditWidgetForOutliner() override;
+	virtual TSharedRef<SWidget> GetCustomOutlinerContent() override;
 	virtual FText GetDisplayName() const override;
 	virtual float GetNodeHeight() const override;
 	virtual FNodePadding GetNodePadding() const override;
diff --git a/Engine/Source/Editor/Sequencer/Private/DisplayNodes/SequencerSpacerNode.h b/Engine/Source/Editor/Sequencer/Private/DisplayNodes/SequencerSpacerNode.h
new file mode 100644
--- /dev/null
+++ b/Engine/Source/Editor/Sequencer/Private/DisplayNodes/SequencerSpacerNode.h
@@ -0,0 +1,46 @@
+// Copyright 1998-2015 Epic Games, Inc. All Rights Reserved.
+
+#pragma once
+
+#include "SequencerDisplayNode.h"
+
+
+class FSequencerNodeTree;
+
+
+/**
+ * A node that displays a category for other nodes
+ */
+class FSequencerSpacerNode
+	: public FSequencerDisplayNode
+{
+public:
+
+	/**
+	 * Create and initialize a new instance.
+	 * 
+	 * @param InParentNode The parent of this node, or nullptr if this is a root node.
+	 * @param InParentTree The tree this node is in.
+	 */
+	FSequencerSpacerNode(float InSize, TSharedPtr<FSequencerDisplayNode> InParentNode, FSequencerNodeTree& InParentTree)
+		: FSequencerDisplayNode(NAME_None, InParentNode, InParentTree)
+		, Size(InSize)
+	{ }
+
+public:
+
+	// FSequencerDisplayNode interface
+
+	virtual bool CanRenameNode() const override { return false; }
+	virtual FText GetDisplayName() const override { return FText(); }
+	virtual float GetNodeHeight() const override { return Size; }
+	virtual FNodePadding GetNodePadding() const override { return FNodePadding(0.f); }
+	virtual ESequencerNode::Type GetType() const override { return ESequencerNode::Spacer; }
+	virtual void SetDisplayName(const FText& NewDisplayName) override { }
+	virtual TSharedRef<SWidget> GenerateContainerWidgetForOutliner(const TSharedRef<SSequencerTreeViewRow>& InRow) override { return SNew(SBox).HeightOverride(Size); }
+
+private:
+
+	/** The size of the spacer */
+	float Size;
+};
