Fixed issue with placementmode module crashing in cooker because asset registry is already shutdown.

diff --git a/Engine/Source/Editor/PlacementMode/Private/PlacementModeModule.cpp b/Engine/Source/Editor/PlacementMode/Private/PlacementModeModule.cpp
--- a/Engine/Source/Editor/PlacementMode/Private/PlacementModeModule.cpp
+++ b/Engine/Source/Editor/PlacementMode/Private/PlacementModeModule.cpp
@@ -201,7 +201,7 @@ public:
 	/**
 	 * Called before the module is unloaded, right before the module object is destroyed.
 	 */
-	virtual void ShutdownModule() override
+	virtual void PreUnloadCallback() override
 	{
 		FEditorModeRegistry::Get().UnregisterMode( FBuiltinEditorModes::EM_Placement );
 
