Non-Unity mac fixes
 - Added missing #includes
 - Removed #pragma once from cpp

diff --git a/Engine/Source/Editor/LevelEditor/Private/LevelViewportLayoutEntity.cpp b/Engine/Source/Editor/LevelEditor/Private/LevelViewportLayoutEntity.cpp
--- a/Engine/Source/Editor/LevelEditor/Private/LevelViewportLayoutEntity.cpp
+++ b/Engine/Source/Editor/LevelEditor/Private/LevelViewportLayoutEntity.cpp
@@ -1,7 +1,5 @@
 // Copyright 1998-2016 Epic Games, Inc. All Rights Reserved.
 
-#pragma once
-
 #include "LevelEditor.h"
 #include "LevelViewportLayoutEntity.h"
 #include "SLevelViewport.h"
diff --git a/Engine/Source/Editor/PlacementMode/Public/IPlacementMode.h b/Engine/Source/Editor/PlacementMode/Public/IPlacementMode.h
--- a/Engine/Source/Editor/PlacementMode/Public/IPlacementMode.h
+++ b/Engine/Source/Editor/PlacementMode/Public/IPlacementMode.h
@@ -1,6 +1,9 @@
 // Copyright 1998-2016 Epic Games, Inc. All Rights Reserved.
 #pragma once
 
+#include "EditorModes.h"
+#include "EdMode.h"
+
 class IPlacementMode : public FEdMode
 {
 public:
diff --git a/Engine/Source/Editor/PlacementMode/Public/IPlacementModeModule.h b/Engine/Source/Editor/PlacementMode/Public/IPlacementModeModule.h
--- a/Engine/Source/Editor/PlacementMode/Public/IPlacementModeModule.h
+++ b/Engine/Source/Editor/PlacementMode/Public/IPlacementModeModule.h
@@ -2,9 +2,12 @@
 #pragma once
 
 #include "ModuleManager.h"
+#include "Editor.h"
+#include "ActorFactories/ActorFactory.h"
 #include "ActorPlacementInfo.h"
 #include "IPlacementMode.h"
 
+
 /**
  * Struct that defines an identifier for a particular placeable item in this module.
  * Only be obtainable through IPlacementModeModule::RegisterPlaceableItem
