Fixed class/actor names in the placement panel not localizing correctly

diff --git a/Engine/Source/Editor/PlacementMode/Public/IPlacementModeModule.h b/Engine/Source/Editor/PlacementMode/Public/IPlacementModeModule.h
--- a/Engine/Source/Editor/PlacementMode/Public/IPlacementModeModule.h
+++ b/Engine/Source/Editor/PlacementMode/Public/IPlacementModeModule.h
@@ -124,22 +124,16 @@ struct FPlaceableItem
 	/** Automatically set this item's display name from its class or asset */
 	void AutoSetDisplayName()
 	{
-		const bool bIsClass = AssetData.GetClass() == UClass::StaticClass();
-		const bool bIsVolume = bIsClass ? CastChecked<UClass>(AssetData.GetAsset())->IsChildOf(AVolume::StaticClass()) : false;
-		const bool bIsActor = bIsClass ? CastChecked<UClass>(AssetData.GetAsset())->IsChildOf(AActor::StaticClass()) : false;
+		UClass* Class = AssetData.GetClass() == UClass::StaticClass() ? Cast<UClass>(AssetData.GetAsset()) : nullptr;
+		const bool bIsVolume = Class && Class->IsChildOf<AVolume>();
 
 		if (Factory && !bIsVolume) // Factories give terrible names for volumes
 		{
 			DisplayName = Factory->GetDisplayName();
 		}
-		else if (bIsActor)
+		else if (Class)
 		{
-			AActor* DefaultActor = CastChecked<AActor>(CastChecked<UClass>(AssetData.GetAsset())->ClassDefaultObject);
-			DisplayName = FText::FromString(FName::NameToDisplayString(DefaultActor->GetClass()->GetName(), false));
-		}
-		else if (bIsClass)
-		{
-			DisplayName = FText::FromString(FName::NameToDisplayString(AssetData.AssetName.ToString(), false));
+			DisplayName = Class->GetDisplayNameText();
 		}
 		else
 		{
