PR #4311: PlacementModeTools shapes searchable and thumbnail (Contributed by [REDACTED_PERSON])

diff --git a/Engine/Source/Editor/PlacementMode/Private/PlacementModeModule.cpp b/Engine/Source/Editor/PlacementMode/Private/PlacementModeModule.cpp
--- a/Engine/Source/Editor/PlacementMode/Private/PlacementModeModule.cpp
+++ b/Engine/Source/Editor/PlacementMode/Private/PlacementModeModule.cpp
@@ -197,6 +197,13 @@ void FPlacementModeModule::StartupModule()
 			50
 		)
 	);
+
+
+	BasicShapeThumbnails.Add(UActorFactoryBasicShape::BasicCube.ToString(), TEXT("ClassThumbnail.Cube"));
+	BasicShapeThumbnails.Add(UActorFactoryBasicShape::BasicSphere.ToString(), TEXT("ClassThumbnail.Sphere"));
+	BasicShapeThumbnails.Add(UActorFactoryBasicShape::BasicCylinder.ToString(), TEXT("ClassThumbnail.Cylinder"));
+	BasicShapeThumbnails.Add(UActorFactoryBasicShape::BasicCone.ToString(), TEXT("ClassThumbnail.Cone"));
+	BasicShapeThumbnails.Add(UActorFactoryBasicShape::BasicPlane.ToString(), TEXT("ClassThumbnail.Plane"));
 }
 
 void FPlacementModeModule::PreUnloadCallback()
@@ -436,7 +443,14 @@ void FPlacementModeModule::RefreshRecentlyPlaced()
 			if (AssetData.IsValid())
 			{
 				UActorFactory* Factory = FindObject<UActorFactory>(nullptr, *RecentlyPlacedItem.Factory);
-				Category->Items.Add(CreateID(), MakeShareable(new FPlaceableItem(Factory, AssetData)));
+				TSharedPtr<FPlaceableItem> Ptr = MakeShareable(new FPlaceableItem(Factory, AssetData));
+				if (FString* FoundThumbnail = BasicShapeThumbnails.Find(RecentlyPlacedItem.ObjectPath))
+				{
+					Ptr->ClassThumbnailBrushOverride = FName(**FoundThumbnail);
+					Ptr->bAlwaysUseGenericThumbnail = true;
+					Ptr->AssetTypeColorOverride = GetBasicShapeColorOverride();
+				}
+				Category->Items.Add(CreateID(), Ptr);
 			}
 		}
 	}
@@ -487,11 +501,11 @@ void FPlacementModeModule::RefreshAllPlaceableClasses()
 	Category->Items.Add(CreateID(), MakeShareable(new FPlaceableItem(*UActorFactoryEmptyActor::StaticClass())));
 	Category->Items.Add(CreateID(), MakeShareable(new FPlaceableItem(*UActorFactoryCharacter::StaticClass())));
 	Category->Items.Add(CreateID(), MakeShareable(new FPlaceableItem(*UActorFactoryPawn::StaticClass())));
-	Category->Items.Add(CreateID(), MakeShareable(new FPlaceableItem(*UActorFactoryBasicShape::StaticClass(), FAssetData(LoadObject<UStaticMesh>(nullptr, *UActorFactoryBasicShape::BasicCube.ToString())), FName("ClassThumbnail.Cube"), GetBasicShapeColorOverride())));
-	Category->Items.Add(CreateID(), MakeShareable(new FPlaceableItem(*UActorFactoryBasicShape::StaticClass(), FAssetData(LoadObject<UStaticMesh>(nullptr, *UActorFactoryBasicShape::BasicSphere.ToString())), FName("ClassThumbnail.Sphere"), GetBasicShapeColorOverride())));
-	Category->Items.Add(CreateID(), MakeShareable(new FPlaceableItem(*UActorFactoryBasicShape::StaticClass(), FAssetData(LoadObject<UStaticMesh>(nullptr, *UActorFactoryBasicShape::BasicCylinder.ToString())), FName("ClassThumbnail.Cylinder"), GetBasicShapeColorOverride())));
-	Category->Items.Add(CreateID(), MakeShareable(new FPlaceableItem(*UActorFactoryBasicShape::StaticClass(), FAssetData(LoadObject<UStaticMesh>(nullptr, *UActorFactoryBasicShape::BasicCone.ToString())), FName("ClassThumbnail.Cone"), GetBasicShapeColorOverride())));
-	Category->Items.Add(CreateID(), MakeShareable(new FPlaceableItem(*UActorFactoryBasicShape::StaticClass(), FAssetData(LoadObject<UStaticMesh>(nullptr, *UActorFactoryBasicShape::BasicPlane.ToString())), FName("ClassThumbnail.Plane"), GetBasicShapeColorOverride())));
+	Category->Items.Add(CreateID(), MakeShareable(new FPlaceableItem(*UActorFactoryBasicShape::StaticClass(), FAssetData(LoadObject<UStaticMesh>(nullptr, *UActorFactoryBasicShape::BasicCube.ToString())), FName("ClassThumbnail.Cube"), GetBasicShapeColorOverride(), TOptional<int32>(), NSLOCTEXT("PlacementMode", "Cube", "Cube"))));
+	Category->Items.Add(CreateID(), MakeShareable(new FPlaceableItem(*UActorFactoryBasicShape::StaticClass(), FAssetData(LoadObject<UStaticMesh>(nullptr, *UActorFactoryBasicShape::BasicSphere.ToString())), FName("ClassThumbnail.Sphere"), GetBasicShapeColorOverride(), TOptional<int32>(), NSLOCTEXT("PlacementMode", "Sphere", "Sphere"))));
+	Category->Items.Add(CreateID(), MakeShareable(new FPlaceableItem(*UActorFactoryBasicShape::StaticClass(), FAssetData(LoadObject<UStaticMesh>(nullptr, *UActorFactoryBasicShape::BasicCylinder.ToString())), FName("ClassThumbnail.Cylinder"), GetBasicShapeColorOverride(), TOptional<int32>(), NSLOCTEXT("PlacementMode", "Cylinder", "Cylinder"))));
+	Category->Items.Add(CreateID(), MakeShareable(new FPlaceableItem(*UActorFactoryBasicShape::StaticClass(), FAssetData(LoadObject<UStaticMesh>(nullptr, *UActorFactoryBasicShape::BasicCone.ToString())), FName("ClassThumbnail.Cone"), GetBasicShapeColorOverride(), TOptional<int32>(), NSLOCTEXT("PlacementMode", "Cone", "Cone"))));
+	Category->Items.Add(CreateID(), MakeShareable(new FPlaceableItem(*UActorFactoryBasicShape::StaticClass(), FAssetData(LoadObject<UStaticMesh>(nullptr, *UActorFactoryBasicShape::BasicPlane.ToString())), FName("ClassThumbnail.Plane"), GetBasicShapeColorOverride(), TOptional<int32>(), NSLOCTEXT("PlacementMode", "Plane", "Plane"))));
 
 	// Make a map of UClasses to ActorFactories that support them
 	const TArray< UActorFactory *>& ActorFactories = GEditor->ActorFactories;
diff --git a/Engine/Source/Editor/PlacementMode/Public/IPlacementModeModule.h b/Engine/Source/Editor/PlacementMode/Public/IPlacementModeModule.h
--- a/Engine/Source/Editor/PlacementMode/Public/IPlacementModeModule.h
+++ b/Engine/Source/Editor/PlacementMode/Public/IPlacementModeModule.h
@@ -126,9 +126,10 @@ struct FPlaceableItem
 	{
 		const bool bIsClass = AssetData.GetClass() == UClass::StaticClass();
 		const bool bIsVolume = bIsClass ? CastChecked<UClass>(AssetData.GetAsset())->IsChildOf(AVolume::StaticClass()) : false;
+		const bool bIsShape = bIsClass ? false : AssetData.GetClass()->IsChildOf(UStaticMesh::StaticClass());
 		const bool bIsActor = bIsClass ? CastChecked<UClass>(AssetData.GetAsset())->IsChildOf(AActor::StaticClass()) : false;
 
-		if (Factory && !bIsVolume) // Factories give terrible names for volumes
+		if (Factory && !bIsVolume && !bIsShape) // Factories give terrible names for volumes and basic shapes
 		{
 			DisplayName = Factory->GetDisplayName();
 		}
diff --git a/Engine/Source/Editor/PlacementMode/Public/PlacementModeModule.h b/Engine/Source/Editor/PlacementMode/Public/PlacementModeModule.h
--- a/Engine/Source/Editor/PlacementMode/Public/PlacementModeModule.h
+++ b/Engine/Source/Editor/PlacementMode/Public/PlacementModeModule.h
@@ -154,6 +154,8 @@ private:
 
 	TArray< TSharedPtr<FExtender> > ContentPaletteFiltersExtenders;
 	TArray< TSharedPtr<FExtender> > PaletteExtenders;
+
+	TMap<FString, FString> BasicShapeThumbnails;
 };
 
 IMPLEMENT_MODULE(FPlacementModeModule, PlacementMode);
\ No newline at end of file
