Fixing crash on toolset reset
#rnx

diff --git a/Engine/Source/Editor/PlacementMode/Private/PlacementMode.cpp b/Engine/Source/Editor/PlacementMode/Private/PlacementMode.cpp
--- a/Engine/Source/Editor/PlacementMode/Private/PlacementMode.cpp
+++ b/Engine/Source/Editor/PlacementMode/Private/PlacementMode.cpp
@@ -56,7 +56,6 @@ void FPlacementMode::Exit()
 	if ( Toolkit.IsValid() )
 	{
 		FToolkitManager::Get().CloseToolkit( Toolkit.ToSharedRef() );
-		Toolkit.Reset();
 	}
 
 	// Call parent implementation
