Roling back 9751429
#rnx

diff --git a/Engine/Source/Editor/PlacementMode/Private/PlacementMode.cpp b/Engine/Source/Editor/PlacementMode/Private/PlacementMode.cpp
--- a/Engine/Source/Editor/PlacementMode/Private/PlacementMode.cpp
+++ b/Engine/Source/Editor/PlacementMode/Private/PlacementMode.cpp
@@ -56,6 +56,7 @@ void FPlacementMode::Exit()
 	if ( Toolkit.IsValid() )
 	{
 		FToolkitManager::Get().CloseToolkit( Toolkit.ToSharedRef() );
+		Toolkit.Reset();
 	}
 
 	// Call parent implementation
