Add if/ensure around potential nullptr from crash reporter in SPlacementModeTools

diff --git a/Engine/Source/Editor/PlacementMode/Private/SPlacementModeTools.cpp b/Engine/Source/Editor/PlacementMode/Private/SPlacementModeTools.cpp
--- a/Engine/Source/Editor/PlacementMode/Private/SPlacementModeTools.cpp
+++ b/Engine/Source/Editor/PlacementMode/Private/SPlacementModeTools.cpp
@@ -602,7 +602,11 @@ FReply SPlacementModeTools::OnKeyDown( const FGeometry& MyGeometry, const FKeyEv
 	if ( InKeyEvent.GetKey() == EKeys::Escape )
 	{
 		FPlacementMode* PlacementEditMode = (FPlacementMode*)GLevelEditorModeTools().GetActiveMode( FBuiltinEditorModes::EM_Placement );
-		PlacementEditMode->StopPlacing();
+		// Catch potential nullptr
+		if (ensureMsgf(PlacementEditMode, TEXT("PlacementEditMode was null, but SPlacementModeTools is still accepting KeyDown events")))
+		{
+			PlacementEditMode->StopPlacing();
+		}
 		Reply = FReply::Handled();
 	}
 
