Placement Browser Style Updates

diff --git a/Engine/Plugins/Editor/GeometryMode/Source/BspMode/Private/BspModeModule.cpp b/Engine/Plugins/Editor/GeometryMode/Source/BspMode/Private/BspModeModule.cpp
--- a/Engine/Plugins/Editor/GeometryMode/Source/BspMode/Private/BspModeModule.cpp
+++ b/Engine/Plugins/Editor/GeometryMode/Source/BspMode/Private/BspModeModule.cpp
@@ -47,7 +47,7 @@ void FBspModeModule::StartupModule()
 
 	IPlacementModeModule& PlacementModeModule = IPlacementModeModule::Get();
 
-	FPlacementCategoryInfo CategoryInfo(LOCTEXT("PlacementMode_Geometry", "Geometry"), "Geometry", TEXT("PMGeometry"), 35);
+	FPlacementCategoryInfo CategoryInfo(LOCTEXT("PlacementMode_Geometry", "Geometry"), FSlateIcon(FAppStyle::GetAppStyleSetName(), "PlacementBrowser.Icons.BSP"), "Geometry", TEXT("PMGeometry"), 35);
 	CategoryInfo.CustomGenerator = []() -> TSharedRef<SWidget> { return SNew(SBspPalette); };
 	PlacementModeModule.RegisterPlacementCategory(CategoryInfo);
 }
diff --git a/Engine/Plugins/Editor/GeometryMode/Source/BspMode/Private/SBspPalette.cpp b/Engine/Plugins/Editor/GeometryMode/Source/BspMode/Private/SBspPalette.cpp
--- a/Engine/Plugins/Editor/GeometryMode/Source/BspMode/Private/SBspPalette.cpp
+++ b/Engine/Plugins/Editor/GeometryMode/Source/BspMode/Private/SBspPalette.cpp
@@ -8,6 +8,7 @@
 #include "Widgets/Images/SImage.h"
 #include "Styling/SlateTypes.h"
 #include "Widgets/Text/STextBlock.h"
+#include "Widgets/Input/SButton.h"
 #include "Widgets/Layout/SBox.h"
 #include "Widgets/Layout/SScrollBorder.h"
 #include "Widgets/Layout/SBorder.h"
@@ -34,19 +35,68 @@ public:
 	}
 };
 
+class SBspButton: public SButton
+{
+	SLATE_BEGIN_ARGS(SBspButton) {}
+	SLATE_END_ARGS()
+
+	void Construct(const FArguments& InArgs, const TSharedPtr<FBspBuilderType>& InBspBuilder)
+	{
+		BspBuilder = InBspBuilder;	
+		SButton::Construct( SButton::FArguments()
+			.ButtonStyle(&FAppStyle::Get(), "PlacementBrowser.Asset")
+			.Cursor( EMouseCursor::GrabHand )
+		);
+	}
+
+	virtual FReply OnMouseButtonDown(const FGeometry& MyGeometry, const FPointerEvent& MouseEvent) override
+	{
+		if ( MouseEvent.GetEffectingButton() == EKeys::LeftMouseButton )
+		{
+			bIsPressed = true;
+
+			return FReply::Handled().DetectDrag( SharedThis( this ), MouseEvent.GetEffectingButton() );
+		}
+
+		return FReply::Unhandled();
+	}
+
+	virtual FReply OnMouseButtonUp(const FGeometry& MyGeometry, const FPointerEvent& MouseEvent) override
+	{
+		SButton::OnMouseButtonUp(MyGeometry, MouseEvent);
+		return FReply::Unhandled();	
+	}
+
+	virtual FReply OnDragDetected(const FGeometry& MyGeometry, const FPointerEvent& MouseEvent) override
+	{
+
+		if ( MouseEvent.IsMouseButtonDown( EKeys::LeftMouseButton ) )
+		{
+			if(BspBuilder.IsValid())
+			{
+				TWeakObjectPtr<UBrushBuilder> ActiveBrushBuilder = GEditor->FindBrushBuilder(BspBuilder->BuilderClass.Get());
+				if (ActiveBrushBuilder.IsValid())
+				{
+					return FReply::Handled().BeginDragDrop(FBrushBuilderDragDropOp::New(ActiveBrushBuilder, BspBuilder->Icon, true /*IsAdditive*/));
+				}
+			}
+		}
+
+		return FReply::Unhandled();
+	}
+
+	TSharedPtr<FBspBuilderType> BspBuilder;
+};
+
 BEGIN_SLATE_FUNCTION_BUILD_OPTIMIZATION
 void SBspPalette::Construct( const FArguments& InArgs )
 {
-	bIsAdditive = true;
-
 	FBspModeModule& BspModeModule = FModuleManager::GetModuleChecked<FBspModeModule>("BspMode");
 
-	TSharedRef<SBspBuilderListView> ListViewWidget = 
+	ListViewWidget = 
 		SNew(SBspBuilderListView)
-		.SelectionMode(ESelectionMode::Single)
 		.ListItemsSource(&BspModeModule.GetBspBuilderTypes())
 		.OnGenerateRow(this, &SBspPalette::MakeListViewWidget)
-		.OnSelectionChanged(this, &SBspPalette::OnSelectionChanged)
 		.ItemHeight(35);
 
 	ChildSlot
@@ -56,89 +106,9 @@ void SBspPalette::Construct( const FArguments& InArgs )
 		+ SVerticalBox::Slot()
 		.FillHeight(1.0f)
 		[
-			SNew(SScrollBorder, ListViewWidget)
+			SNew(SScrollBorder, ListViewWidget.ToSharedRef())
 			[
-				ListViewWidget
-			]
-		]
-
-		+ SVerticalBox::Slot()
-		.AutoHeight()
-		[
-			SNew( SHorizontalBox )
-
-			+ SHorizontalBox::Slot()
-			.AutoWidth()
-			.Padding( 3.0f )
-			.HAlign(HAlign_Right)
-			.VAlign(VAlign_Center)
-			[
-				SNew(SCheckBox)
-				.Style(FEditorStyle::Get(), "Toolbar.RadioButton")
-				.ToolTipText(LOCTEXT("BspModeAdditiveTooltip", "Place brushes in additive mode."))
-				.OnCheckStateChanged(this, &SBspPalette::OnAdditiveModeButtonClicked)
-				.IsChecked(this, &SBspPalette::IsAdditiveModeChecked)
-				.ForegroundColor(FCoreStyle::Get().GetSlateColor("Foreground"))
-				.Content()
-				[
-					SNew(SHorizontalBox)
-					+SHorizontalBox::Slot()
-					.AutoWidth()
-					.HAlign(HAlign_Left)
-					.VAlign(VAlign_Center)
-					[
-						SNew(SImage)
-						.Image(this, &SBspPalette::GetAdditiveModeImage)
-					]
-					+SHorizontalBox::Slot()
-					.AutoWidth()
-					.HAlign(HAlign_Center)
-					.VAlign(VAlign_Center)
-					.Padding(2.0f)
-					[
-						SNew(STextBlock)
-						.Text(LOCTEXT("BspModeAdd", "Add"))
-						.Font(FCoreStyle::Get().GetFontStyle("Toolbar.Label.Font"))
-						.ShadowOffset(FVector2D::UnitVector)
-					]
-				]
-			]
-
-			+ SHorizontalBox::Slot()
-			.Padding( 3.0f )
-			.AutoWidth()
-			.HAlign(HAlign_Right)
-			.VAlign(VAlign_Center)
-			[
-				SNew(SCheckBox)
-				.Style(FEditorStyle::Get(), "Toolbar.RadioButton")
-				.ToolTipText(LOCTEXT("BspModeSubtractiveTooltip", "Place brushes in subtractive mode."))
-				.OnCheckStateChanged(this, &SBspPalette::OnSubtractiveModeButtonClicked)
-				.IsChecked(this, &SBspPalette::IsSubtractiveModeChecked)
-				.ForegroundColor(FCoreStyle::Get().GetSlateColor("Foreground"))
-				.Content()
-				[
-					SNew(SHorizontalBox)
-					+SHorizontalBox::Slot()
-					.AutoWidth()
-					.HAlign(HAlign_Left)
-					.VAlign(VAlign_Center)
-					[
-						SNew(SImage)
-						.Image(this, &SBspPalette::GetSubtractiveModeImage)
-					]
-					+SHorizontalBox::Slot()
-					.AutoWidth()
-					.HAlign(HAlign_Left)
-					.VAlign(VAlign_Center)
-					.Padding(2.0f)
-					[
-						SNew(STextBlock)
-						.Text(LOCTEXT("BspModeSubtract", "Subtract"))
-						.Font(FCoreStyle::Get().GetFontStyle("Toolbar.Label.Font"))
-						.ShadowOffset(FVector2D::UnitVector)
-					]
-				]
+				ListViewWidget.ToSharedRef()
 			]
 		]
 	];
@@ -152,51 +122,62 @@ TSharedRef<ITableRow> SBspPalette::MakeListViewWidget(TSharedPtr<FBspBuilderType
 
 	TSharedRef< STableRow<TSharedPtr<FBspBuilderType>> > TableRowWidget = 
 		SNew( STableRow<TSharedPtr<FBspBuilderType>>, OwnerTable )
-		.Style( FBspModeStyle::Get(), "BspMode.TableRow" )
-		.OnDragDetected( this, &SBspPalette::OnDraggingListViewWidget );
+		.Style(&FAppStyle::Get(), "PlacementBrowser.PlaceableItemRow")
+		.Padding(FMargin(8.f, 2.f, 12.f, 2.f));
+
 
 	TSharedRef<SWidget> Content = 
-		SNew(SBorder)
-		.BorderImage(FCoreStyle::Get().GetBrush("NoBorder"))
-		.Padding(0)
-		.ToolTip(FEditorClassUtils::GetTooltip(ABrush::StaticClass(), BspBuilder->ToolTipText))
-		.Cursor( EMouseCursor::GrabHand )
-		[
-			SNew(SHorizontalBox)
+		SNew(SOverlay)
 
-			// Icon
-			+SHorizontalBox::Slot()
-			.AutoWidth()
+		+SOverlay::Slot()
+		[
+			SNew(SBorder)
+			.BorderImage( FAppStyle::Get().GetBrush("PlacementBrowser.Asset.Background"))
+			.Cursor( EMouseCursor::GrabHand )
+			.Padding(0)
 			[
-				SNew(SBorder)
-				.Padding(4.0f)
-				.BorderImage(FBspModeStyle::Get().GetBrush("BspMode.ThumbnailShadow"))
+
+				SNew( SHorizontalBox )
+
+				+ SHorizontalBox::Slot()
+				.Padding(8.0f, 4.f)
+				.AutoWidth()
+				.HAlign(HAlign_Center)
+				.VAlign(VAlign_Center)
 				[
-					SNew(SBox)
-					.WidthOverride(35.0f)
-					.HeightOverride(35.0f)
+					SNew( SBox )
+					.WidthOverride(40)
+					.HeightOverride(40)
 					[
-						SNew(SBorder)
-						.BorderImage(FBspModeStyle::Get().GetBrush("BspMode.ThumbnailBackground"))
-						.HAlign(HAlign_Center)
-						.VAlign(VAlign_Center)
-						[
-							SNew(SImage)
-							.Image(BspBuilder->Icon)
-						]
+
+						SNew(SImage)
+						.Image(BspBuilder->Icon)
 					]
 				]
-			]
 
-			+SHorizontalBox::Slot()
-			.HAlign(HAlign_Left)
-			.VAlign(VAlign_Center)
-			.FillWidth(1.f)
-			[
-				SNew(STextBlock)
-				.TextStyle(FBspModeStyle::Get(), "BspMode.ThumbnailText")
-				.Text(BspBuilder->Text)
+				+ SHorizontalBox::Slot()
+				.VAlign(VAlign_Fill)
+				.Padding(0)
+				[
+
+					SNew(SBorder)
+					.BorderImage(FAppStyle::Get().GetBrush("PlacementBrowser.Asset.LabelBack"))
+					.Padding(FMargin(9, 0, 0, 1))
+					.VAlign(VAlign_Center)
+					[
+							SNew( STextBlock )
+							.TextStyle( FEditorStyle::Get(), "PlacementBrowser.Asset.Name" )
+							.Text(BspBuilder->Text)
+					]
+				]
 			]
+		]
+
+
+		+SOverlay::Slot()
+		[
+			SNew(SBspButton, BspBuilder)
+			.ToolTip(FEditorClassUtils::GetTooltip(ABrush::StaticClass(), BspBuilder->ToolTipText))
 		];
 
 	TableRowWidget->SetContent(Content);
@@ -204,61 +185,4 @@ TSharedRef<ITableRow> SBspPalette::MakeListViewWidget(TSharedPtr<FBspBuilderType
 	return TableRowWidget;
 }
 
-void SBspPalette::OnSelectionChanged(TSharedPtr<FBspBuilderType> BspBuilder, ESelectInfo::Type SelectionType)
-{
-	if(BspBuilder.IsValid())
-	{
-		ActiveBrushBuilder = GEditor->FindBrushBuilder(BspBuilder->BuilderClass.Get());
-	}
-}
-
-FReply SBspPalette::OnDraggingListViewWidget(const FGeometry& MyGeometry, const FPointerEvent& MouseEvent)
-{
-	if ( MouseEvent.IsMouseButtonDown( EKeys::LeftMouseButton ) )
-	{
-		if (ActiveBrushBuilder.IsValid())
-		{
-			FBspModeModule& BspModeModule = FModuleManager::GetModuleChecked<FBspModeModule>("BspMode");
-			TSharedPtr<FBspBuilderType> BspBuilder = BspModeModule.FindBspBuilderType(ActiveBrushBuilder.Get()->GetClass());
-			if(BspBuilder.IsValid())
-			{
-				// We have an active brush builder, start a drag-drop
-				return FReply::Handled().BeginDragDrop(FBrushBuilderDragDropOp::New(ActiveBrushBuilder, BspBuilder->Icon, bIsAdditive));
-			}
-		}
-	}
-
-	return FReply::Unhandled();
-}
-
-void SBspPalette::OnAdditiveModeButtonClicked(ECheckBoxState CheckType)
-{
-	bIsAdditive = CheckType == ECheckBoxState::Checked;
-}
-
-void SBspPalette::OnSubtractiveModeButtonClicked(ECheckBoxState CheckType)
-{
-	bIsAdditive = CheckType != ECheckBoxState::Checked;
-}
-
-ECheckBoxState SBspPalette::IsAdditiveModeChecked() const
-{
-	return bIsAdditive ? ECheckBoxState::Checked : ECheckBoxState::Unchecked;
-}
-
-ECheckBoxState SBspPalette::IsSubtractiveModeChecked() const
-{
-	return !bIsAdditive ? ECheckBoxState::Checked : ECheckBoxState::Unchecked;
-}
-
-const FSlateBrush* SBspPalette::GetAdditiveModeImage() const
-{
-	return FBspModeStyle::Get().GetBrush("BspMode.CSGAdd.Small");
-}
-
-const FSlateBrush* SBspPalette::GetSubtractiveModeImage() const
-{
-	return FBspModeStyle::Get().GetBrush("BspMode.CSGSubtract.Small");
-}
-
 #undef LOCTEXT_NAMESPACE
diff --git a/Engine/Plugins/Editor/GeometryMode/Source/BspMode/Private/SBspPalette.h b/Engine/Plugins/Editor/GeometryMode/Source/BspMode/Private/SBspPalette.h
--- a/Engine/Plugins/Editor/GeometryMode/Source/BspMode/Private/SBspPalette.h
+++ b/Engine/Plugins/Editor/GeometryMode/Source/BspMode/Private/SBspPalette.h
@@ -13,6 +13,7 @@ class IDetailsView;
 class ITableRow;
 class STableViewBase;
 class UBrushBuilder;
+class SBspBuilderListView;
 struct FSlateBrush;
 enum class ECheckBoxState : uint8;
 
@@ -29,37 +30,8 @@ private:
 	/** Make a widget for the list view display */
 	TSharedRef<ITableRow> MakeListViewWidget(TSharedPtr<struct FBspBuilderType> BspBuilder, const TSharedRef<STableViewBase>& OwnerTable);
 
-	/** Delegate for when the list view selection changes */
-	void OnSelectionChanged(TSharedPtr<FBspBuilderType> BspBuilder, ESelectInfo::Type SelectionType);
 
-	/** Begin dragging a list widget */
-	FReply OnDraggingListViewWidget(const FGeometry& MyGeometry, const FPointerEvent& MouseEvent);
+	/** List View for the BSP Builder Types */
+	TSharedPtr<class SBspBuilderListView> ListViewWidget;
 
-	/** Select additive mode */
-	void OnAdditiveModeButtonClicked(ECheckBoxState CheckType);
-
-	/** Select subtractive mode */
-	void OnSubtractiveModeButtonClicked(ECheckBoxState CheckType);
-
-	/** @return the check state of the additive mode checkbox */
-	ECheckBoxState IsAdditiveModeChecked() const;
-
-	/** @return the check state of the subtractive mode checkbox */
-	ECheckBoxState IsSubtractiveModeChecked() const;
-
-	/** Get the image for additive mode */
-	const FSlateBrush* GetAdditiveModeImage() const;
-
-	/** Get the image for subtractive mode */
-	const FSlateBrush* GetSubtractiveModeImage() const;
-
-private:
-	/** Property view for brush options */
-	TSharedPtr<class IDetailsView> BrushOptionView;
-
-	/** Brush builder currently active */
-	TWeakObjectPtr<UBrushBuilder> ActiveBrushBuilder;
-
-	/** Additive or subtractive mode */
-	bool bIsAdditive;
 };
diff --git a/Engine/Source/Editor/EditorStyle/Private/StarshipStyle.cpp b/Engine/Source/Editor/EditorStyle/Private/StarshipStyle.cpp
--- a/Engine/Source/Editor/EditorStyle/Private/StarshipStyle.cpp
+++ b/Engine/Source/Editor/EditorStyle/Private/StarshipStyle.cpp
@@ -13,6 +13,7 @@
 	#include "PlatformInfo.h"
 #endif
 #include "Styling/ToolBarStyle.h"
+#include "Styling/SegmentedControlStyle.h"
 #include "Styling/StyleColors.h"
 #include "Framework/Application/SlateApplication.h"
 
@@ -4272,7 +4273,7 @@ void FStarshipEditorStyle::FStyle::SetupLevelEditorStyle()
 		Set( "EditorViewport.UnlitMode",          	  new IMAGE_BRUSH_SVG("Starship/Common/UnlitCube", Icon16x16 ) );
 		Set( "EditorViewport.WireframeMode",      	  new IMAGE_BRUSH_SVG("Starship/Common/BrushWireframe", Icon16x16 ) );
 		Set( "EditorViewport.DetailLightingMode", 	  new IMAGE_BRUSH_SVG("Starship/Common/DetailLighting", Icon16x16 ) );
-		Set( "EditorViewport.LightingOnlyMode",   	  new IMAGE_BRUSH_SVG("Starship/Common/LightingOnly", Icon16x16 ) );
+		Set( "EditorViewport.LightingOnlyMode",   	  new IMAGE_BRUSH_SVG("Starship/Common/LightBulb", Icon16x16 ) );
 
 		Set( "EditorViewport.LightComplexityMode", new IMAGE_BRUSH( "Icons/icon_ViewMode_LightComplexity_16px", Icon16x16 ) );
 		Set( "EditorViewport.ShaderComplexityMode", new IMAGE_BRUSH( "Icons/icon_ViewMode_Shadercomplexity_16px", Icon16x16 ) );
@@ -4472,13 +4473,16 @@ void FStarshipEditorStyle::FStyle::SetupLevelEditorStyle()
 		Set( "PlacementBrowser.AssetToolTip.AssetPath", FTextBlockStyle(NormalText) .SetFont( DEFAULT_FONT( "Regular", 8 ) ) );
 
 		Set( "PlacementBrowser.Asset", FButtonStyle( Button )
-			.SetNormal( FSlateNoResource() )
-			.SetHovered( IMAGE_BRUSH( "Common/Selection", Icon8x8, FLinearColor( 1.0f, 1.0f, 1.0f, 0.1f ) ) )
-			.SetPressed( IMAGE_BRUSH( "Common/Selection", Icon8x8, SelectionColor ) )
+			.SetNormal( FSlateRoundedBoxBrush(FLinearColor::Transparent, 6.0f, FStyleColors::Dropdown, 1.0f) )
+			.SetHovered( FSlateRoundedBoxBrush(FLinearColor::Transparent, 6.0f, FStyleColors::Hover, 1.0f) )
+			.SetPressed( FSlateRoundedBoxBrush(FLinearColor::Transparent, 6.0f, FStyleColors::Primary, 1.0f) )
 			.SetNormalPadding( 0 )
 			.SetPressedPadding( 0 )
 			);
 
+		Set( "PlacementBrowser.Asset.Background", new FSlateRoundedBoxBrush(FStyleColors::Recessed, 6.f));
+		Set( "PlacementBrowser.Asset.LabelBack", new BOX_BRUSH("Starship/PlacementBrowser/LabelBack_18x", 6.f/18.f, FStyleColors::Dropdown));
+
 		FLinearColor DimBackground = FLinearColor( FColor( 64, 64, 64 ) );
 		FLinearColor DimBackgroundHover = FLinearColor( FColor( 50, 50, 50 ) );
 		FLinearColor DarkBackground = FLinearColor( FColor( 42, 42, 42 ) );
@@ -4519,6 +4523,67 @@ void FStarshipEditorStyle::FStyle::SetupLevelEditorStyle()
 		Set( "PlacementBrowser.ShowCollections", new IMAGE_BRUSH( "Icons/icon_Placement_Collections_20px", Icon20x20 ) );
 		Set( "PlacementBrowser.ShowCollections.Small", new IMAGE_BRUSH( "Icons/icon_Placement_Collections_20px", Icon20x20 ) );
 
+
+		const FTableRowStyle PlaceItemTableRowStyle = FTableRowStyle()
+			.SetEvenRowBackgroundBrush(FSlateNoResource())
+			.SetEvenRowBackgroundHoveredBrush(FSlateNoResource())
+
+			.SetOddRowBackgroundBrush(FSlateNoResource())
+			.SetOddRowBackgroundHoveredBrush(FSlateNoResource())
+
+			.SetSelectorFocusedBrush(BORDER_BRUSH("Common/Selector", FMargin(4.f / 16.f), SelectorColor))
+
+			.SetActiveBrush(FSlateNoResource())
+			.SetActiveHoveredBrush(FSlateNoResource())
+
+			.SetInactiveBrush(FSlateNoResource())
+			.SetInactiveHoveredBrush(FSlateNoResource())
+
+			.SetActiveHighlightedBrush(FSlateNoResource())
+			.SetInactiveHighlightedBrush(FSlateNoResource())
+
+			.SetTextColor(FStyleColors::Foreground)
+			.SetSelectedTextColor(FStyleColors::Foreground)
+
+			.SetDropIndicator_Above(BOX_BRUSH("Common/DropZoneIndicator_Above", FMargin(10.0f / 16.0f, 10.0f / 16.0f, 0, 0), SelectionColor))
+			.SetDropIndicator_Onto(BOX_BRUSH("Common/DropZoneIndicator_Onto", FMargin(4.0f / 16.0f), SelectionColor))
+			.SetDropIndicator_Below(BOX_BRUSH("Common/DropZoneIndicator_Below", FMargin(10.0f / 16.0f, 0, 0, 10.0f / 16.0f), SelectionColor));
+
+		Set("PlacementBrowser.PlaceableItemRow", PlaceItemTableRowStyle);
+
+
+		const FCheckBoxStyle PlacementSegmentedBox = FCheckBoxStyle()
+			.SetCheckBoxType(ESlateCheckBoxType::ToggleButton)
+			.SetUncheckedImage(FSlateNoResource())
+			.SetUncheckedHoveredImage(FSlateNoResource())
+			.SetUncheckedPressedImage(FSlateNoResource())
+			.SetCheckedImage(FSlateNoResource())
+			.SetCheckedHoveredImage(FSlateNoResource())
+			.SetCheckedPressedImage(FSlateNoResource())
+			.SetForegroundColor(FStyleColors::Foreground)
+			.SetHoveredForegroundColor(FStyleColors::ForegroundHover)
+			.SetPressedForegroundColor(FStyleColors::ForegroundHover)
+			.SetCheckedForegroundColor(FStyleColors::Primary)
+			.SetCheckedHoveredForegroundColor(FStyleColors::Primary)
+			.SetCheckedPressedForegroundColor(FStyleColors::Primary)
+			.SetPadding(FMargin(6.f, 2.f));
+
+		Set("PlacementBrowser.CategoryControl", FSegmentedControlStyle()
+			.SetControlStyle(PlacementSegmentedBox)
+			.SetFirstControlStyle(PlacementSegmentedBox)
+			.SetLastControlStyle(PlacementSegmentedBox)
+		);
+
+		Set("PlacementBrowser.Icons.Recent",        new IMAGE_BRUSH_SVG("Starship/Common/Recent",         Icon16x16));
+		Set("PlacementBrowser.Icons.Basic",         new IMAGE_BRUSH_SVG("Starship/Common/Basic",          Icon16x16));
+		Set("PlacementBrowser.Icons.Lights",        new IMAGE_BRUSH_SVG("Starship/Common/LightBulb",      Icon16x16));
+		Set("PlacementBrowser.Icons.Cinematics",    new IMAGE_BRUSH_SVG("Starship/Common/Cinematics",     Icon16x16));
+		Set("PlacementBrowser.Icons.VisualEffects", new IMAGE_BRUSH_SVG("Starship/Common/VisualEffects",  Icon16x16));
+		Set("PlacementBrowser.Icons.BSP",           new IMAGE_BRUSH_SVG("Starship/Common/Geometry",       Icon16x16));
+		Set("PlacementBrowser.Icons.Volumes",       new IMAGE_BRUSH_SVG("Starship/Common/Volumes",        Icon16x16));
+		Set("PlacementBrowser.Icons.All",           new IMAGE_BRUSH_SVG("Starship/Common/AllClasses",     Icon16x16));
+		Set("PlacementBrowser.Icons.Testing",       new IMAGE_BRUSH_SVG("Starship/Common/Test",           Icon16x16));
+
 		Set( "ContentPalette.ShowAllPlaceables", new IMAGE_BRUSH( "Icons/icon_Placement_FilterAll_20px", Icon20x20 ) );
 		Set( "ContentPalette.ShowAllPlaceables.Small", new IMAGE_BRUSH( "Icons/icon_Placement_FilterAll_20px", Icon20x20 ) );
 		Set( "ContentPalette.ShowProps", new IMAGE_BRUSH( "Icons/icon_Placement_FilterProps_20px", Icon20x20 ) );
diff --git a/Engine/Source/Editor/PlacementMode/Private/SPlacementModeTools.cpp b/Engine/Source/Editor/PlacementMode/Private/SPlacementModeTools.cpp
--- a/Engine/Source/Editor/PlacementMode/Private/SPlacementModeTools.cpp
+++ b/Engine/Source/Editor/PlacementMode/Private/SPlacementModeTools.cpp
@@ -162,57 +162,75 @@ void SPlacementAssetEntry::Construct(const FArguments& InArgs, const TSharedPtr<
 	}
 
 	ChildSlot
+	.Padding(FMargin(8.f, 2.f, 12.f, 2.f))
 	[
-		SNew( SBorder )
-		.BorderImage( this, &SPlacementAssetEntry::GetBorder )
-		.Cursor( EMouseCursor::GrabHand )
-		.ToolTip( AssetEntryToolTip )
-		[
-			SNew( SHorizontalBox )
 
-			+ SHorizontalBox::Slot()
-			.Padding( 0 )
-			.AutoWidth()
+		SNew(SOverlay)
+
+		+SOverlay::Slot()
+		[
+			SNew(SBorder)
+			.BorderImage( FAppStyle::Get().GetBrush("PlacementBrowser.Asset.Background"))
+			.Cursor( EMouseCursor::GrabHand )
+			.ToolTip( AssetEntryToolTip )
+			.Padding(0)
 			[
-				// Drop shadow border
-				SNew( SBorder )
-				.Padding( 4 )
-				.BorderImage( FEditorStyle::GetBrush( "ContentBrowser.ThumbnailShadow" ) )
+
+				SNew( SHorizontalBox )
+
+				+ SHorizontalBox::Slot()
+				.Padding(8.0f, 4.f)
+				.AutoWidth()
+				.HAlign(HAlign_Center)
+				.VAlign(VAlign_Center)
 				[
 					SNew( SBox )
-					.WidthOverride( 35 )
-					.HeightOverride( 35 )
+					.WidthOverride(40)
+					.HeightOverride(40)
 					[
 						SNew( SPlacementAssetThumbnail, Item->AssetData )
 						.ClassThumbnailBrushOverride( Item->ClassThumbnailBrushOverride )
 						.AlwaysUseGenericThumbnail( Item->bAlwaysUseGenericThumbnail )
-						.AssetTypeColorOverride( Item->AssetTypeColorOverride )
+						.AssetTypeColorOverride( FLinearColor::Transparent )
 					]
 				]
-			]
 
-			+ SHorizontalBox::Slot()
-			.VAlign(VAlign_Center)
-			.Padding(2, 0, 4, 0)
-			[
-				SNew( SVerticalBox )
-				+SVerticalBox::Slot()
-				.Padding(0, 0, 0, 1)
-				.AutoHeight()
+				+ SHorizontalBox::Slot()
+				.VAlign(VAlign_Fill)
+				.Padding(0)
 				[
-					SNew( STextBlock )
-					.TextStyle( FEditorStyle::Get(), "PlacementBrowser.Asset.Name" )
-					.Text( Item->DisplayName )
-					.HighlightText(InArgs._HighlightText)
+
+					SNew(SBorder)
+					.BorderImage(FAppStyle::Get().GetBrush("PlacementBrowser.Asset.LabelBack"))
+					[
+						SNew( SHorizontalBox)
+						+SHorizontalBox::Slot()
+						.Padding(9, 0, 0, 1)
+						.VAlign(VAlign_Center)
+						[
+							SNew( STextBlock )
+							.TextStyle( FEditorStyle::Get(), "PlacementBrowser.Asset.Name" )
+							.Text( Item->DisplayName )
+							.HighlightText(InArgs._HighlightText)
+						]
+
+						+ SHorizontalBox::Slot()
+						.VAlign(VAlign_Center)
+						.AutoWidth()
+						[
+							DocWidget
+						]
+					]
 				]
 			]
+		]
 
-			+ SHorizontalBox::Slot()
-			.VAlign(VAlign_Center)
-			.AutoWidth()
-			[
-				DocWidget
-			]
+		+SOverlay::Slot()
+		[
+			SNew(SBorder)
+			.BorderImage( this, &SPlacementAssetEntry::GetBorder )
+			.Cursor( EMouseCursor::GrabHand )
+			.ToolTip( AssetEntryToolTip )
 		]
 	];
 }
@@ -312,20 +330,19 @@ void SPlacementModeTools::Construct( const FArguments& InArgs )
 		FPlacementAssetEntryTextFilter::FItemToStringArray::CreateStatic(&PlacementViewFilter::GetBasicStrings)
 		));
 
-	TSharedRef<SVerticalBox> Tabs = SNew(SVerticalBox).Visibility(this, &SPlacementModeTools::GetTabsVisibility);
+	SAssignNew(CategoryFilterPtr, SSegmentedControl<FName>)
+	.Style(&FAppStyle::Get(), "PlacementBrowser.CategoryControl")
+    .OnValueChanged(this, &SPlacementModeTools::OnCategoryChanged)
+    .Value_Lambda( [this] { return ActiveTabName; } );
 
-	// Populate the tabs and body from the defined placeable items
+	// Populate the categories and body from the defined placeable items
 	IPlacementModeModule& PlacementModeModule = IPlacementModeModule::Get();
 
 	TArray<FPlacementCategoryInfo> Categories;
 	PlacementModeModule.GetSortedCategories(Categories);
 	for (const FPlacementCategoryInfo& Category : Categories)
 	{
-		Tabs->AddSlot()
-		.AutoHeight()
-		[
-			CreatePlacementGroupTab(Category)
-		];
+		CategoryFilterPtr->AddSlot(Category.UniqueHandle).Icon(Category.DisplayIcon.GetIcon());// .Text(Category.DisplayName);
 	}
 
 	TSharedRef<SScrollBar> ScrollBar = SNew(SScrollBar)
@@ -336,68 +353,87 @@ void SPlacementModeTools::Construct( const FArguments& InArgs )
 		SNew( SVerticalBox )
 
 		+ SVerticalBox::Slot()
-		.Padding(4)
 		.AutoHeight()
 		[
-			SAssignNew( SearchBoxPtr, SSearchBox )
-			.HintText(NSLOCTEXT("PlacementMode", "SearchPlaceables", "Search Classes"))
-			.OnTextChanged(this, &SPlacementModeTools::OnSearchChanged)
-			.OnTextCommitted(this, &SPlacementModeTools::OnSearchCommitted)
+			SNew(SBorder)
+			.BorderImage(FEditorStyle::GetBrush("ToolPanel.GroupBorder"))
+			.Padding(8)
+			[
+				SAssignNew( SearchBoxPtr, SSearchBox )
+				.HintText(NSLOCTEXT("PlacementMode", "SearchPlaceables", "Search Classes"))
+				.OnTextChanged(this, &SPlacementModeTools::OnSearchChanged)
+				.OnTextCommitted(this, &SPlacementModeTools::OnSearchCommitted)
+			]
+		]
+
+		+ SVerticalBox::Slot()
+		.AutoHeight()
+		[
+			SNew(SBorder)
+			.BorderImage(FEditorStyle::GetBrush("ToolPanel.GroupBorder"))
+			.Padding(FMargin(8.f, 6.f, 8.f, 8.f))
+			.HAlign(HAlign_Center)
+			.Visibility(this, &SPlacementModeTools::GetTabsVisibility)
+			[
+				CategoryFilterPtr.ToSharedRef()
+			]
 		]
 
+
 		+ SVerticalBox::Slot()
-		.Padding( 0 )
+		.AutoHeight()
 		[
-			SNew(SHorizontalBox)
+			SNew(SBorder)
+			.BorderImage(FEditorStyle::GetBrush("ToolPanel.GroupBorder"))
+			.Padding(FMargin(8.f, 6.f, 8.f, 8.f))
+			.HAlign(HAlign_Center)
+			.Visibility(this, &SPlacementModeTools::GetTabsVisibility)
+			[
+				SAssignNew(FilterLabelPtr, STextBlock)
+				.Text(NSLOCTEXT("PlacementMode", "CategoryLabel", "CategoryLabel"))
+				.Font(FAppStyle::Get().GetFontStyle("SmallFontBold"))
+				.TransformPolicy(ETextTransformPolicy::ToUpper)
+			]
+		]
 
-			+ SHorizontalBox::Slot()
-			.AutoWidth()
+		+ SVerticalBox::Slot()
+		.Padding(FMargin(0.0f, 3.f))
+		[
+			SNew(SOverlay)
+
+			+ SOverlay::Slot()
+			.HAlign(HAlign_Center)
+			.VAlign(VAlign_Fill)
+			.Padding(12.f)
 			[
-				Tabs
+				SNew(STextBlock)
+				.Text(NSLOCTEXT("PlacementMode", "NoResultsFound", "No Results Found"))
+				.Visibility(this, &SPlacementModeTools::GetFailedSearchVisibility)
 			]
 
-			+ SHorizontalBox::Slot()
+			+ SOverlay::Slot()
 			[
-				SNew(SBorder)
-				.Padding(FMargin(3))
-				.BorderImage(FEditorStyle::GetBrush("ToolPanel.DarkGroupBorder"))
-				[
-					SNew(SOverlay)
+				SAssignNew(CustomContent, SBox)
+			]
 
-					+ SOverlay::Slot()
-					.HAlign(HAlign_Center)
-					.VAlign(VAlign_Fill)
-					[
-						SNew(STextBlock)
-						.Text(NSLOCTEXT("PlacementMode", "NoResultsFound", "No Results Found"))
-						.Visibility(this, &SPlacementModeTools::GetFailedSearchVisibility)
-					]
+			+ SOverlay::Slot()
+			[
+				SAssignNew(DataDrivenContent, SBox)
+				[
+					SNew(SHorizontalBox)
 
-					+ SOverlay::Slot()
+					+ SHorizontalBox::Slot()
 					[
-						SAssignNew(CustomContent, SBox)
+						SAssignNew(ListView, SListView<TSharedPtr<FPlaceableItem>>)
+						.ListItemsSource(&FilteredItems)
+						.OnGenerateRow(this, &SPlacementModeTools::OnGenerateWidgetForItem)
+						.ExternalScrollbar(ScrollBar)
 					]
 
-					+ SOverlay::Slot()
+					+ SHorizontalBox::Slot()
+					.AutoWidth()
 					[
-						SAssignNew(DataDrivenContent, SBox)
-						[
-							SNew(SHorizontalBox)
-
-							+ SHorizontalBox::Slot()
-							[
-								SAssignNew(ListView, SListView<TSharedPtr<FPlaceableItem>>)
-								.ListItemsSource(&FilteredItems)
-								.OnGenerateRow(this, &SPlacementModeTools::OnGenerateWidgetForItem)
-								.ExternalScrollbar(ScrollBar)
-							]
-
-							+ SHorizontalBox::Slot()
-							.AutoWidth()
-							[
-								ScrollBar
-							]
-						]
+						ScrollBar
 					]
 				]
 			]
@@ -468,32 +504,28 @@ void SPlacementModeTools::UpdateFilteredItems()
 
 		CustomContent->SetVisibility(EVisibility::Visible);
 		DataDrivenContent->SetVisibility(EVisibility::Collapsed);
+
+		FilterLabelPtr->SetText(Category->DisplayName);
 	}
 	else
 	{
 		FilteredItems.Reset();
 
-		const FPlacementCategoryInfo* CategoryInfo = PlacementModeModule.GetRegisteredPlacementCategory(GetActiveTab());
-		if (!ensure(CategoryInfo))
-		{
-			return;
-		}
-		
 		if (IsSearchActive())
 		{
 			auto Filter = [&](const TSharedPtr<FPlaceableItem>& Item) { return SearchTextFilter->PassesFilter(*Item); };
-			PlacementModeModule.GetFilteredItemsForCategory(CategoryInfo->UniqueHandle, FilteredItems, Filter);
+			PlacementModeModule.GetFilteredItemsForCategory(Category->UniqueHandle, FilteredItems, Filter);
 			
-			if (CategoryInfo->bSortable)
+			if (Category->bSortable)
 			{
 				FilteredItems.Sort(&FSortPlaceableItems::SortItemsByName);
 			}
 		}
 		else
 		{
-			PlacementModeModule.GetItemsForCategory(CategoryInfo->UniqueHandle, FilteredItems);
+			PlacementModeModule.GetItemsForCategory(Category->UniqueHandle, FilteredItems);
 
-			if (CategoryInfo->bSortable)
+			if (Category->bSortable)
 			{
 				FilteredItems.Sort(&FSortPlaceableItems::SortItemsByOrderThenName);
 			}
@@ -502,6 +534,7 @@ void SPlacementModeTools::UpdateFilteredItems()
 		CustomContent->SetVisibility(EVisibility::Collapsed);
 		DataDrivenContent->SetVisibility(EVisibility::Visible);
 		ListView->RequestListRefresh();
+		FilterLabelPtr->SetText(Category->DisplayName);
 	}
 }
 
@@ -532,12 +565,24 @@ EVisibility SPlacementModeTools::GetTabsVisibility() const
 TSharedRef<ITableRow> SPlacementModeTools::OnGenerateWidgetForItem(TSharedPtr<FPlaceableItem> InItem, const TSharedRef<STableViewBase>& OwnerTable)
 {
 	return SNew(STableRow<TSharedPtr<FPlaceableItem>>, OwnerTable)
+		.Style(&FAppStyle::Get(), "PlacementBrowser.PlaceableItemRow")
 		[
 			SNew(SPlacementAssetEntry, InItem.ToSharedRef())
 			.HighlightText(this, &SPlacementModeTools::GetHighlightText)
 		];
 }
 
+void SPlacementModeTools::OnCategoryChanged(FName InCategory)
+{
+	if (InCategory != ActiveTabName)
+	{
+		ActiveTabName = InCategory;
+		IPlacementModeModule::Get().RegenerateItemsForCategory(ActiveTabName);
+
+		bNeedsUpdate = true;
+	}
+}
+
 void SPlacementModeTools::OnPlacementTabChanged( ECheckBoxState NewState, FName CategoryName )
 {
 	if ( NewState == ECheckBoxState::Checked )
diff --git a/Engine/Source/Editor/PlacementMode/Private/SPlacementModeTools.h b/Engine/Source/Editor/PlacementMode/Private/SPlacementModeTools.h
--- a/Engine/Source/Editor/PlacementMode/Private/SPlacementModeTools.h
+++ b/Engine/Source/Editor/PlacementMode/Private/SPlacementModeTools.h
@@ -9,6 +9,7 @@
 #include "Widgets/SCompoundWidget.h"
 #include "ActorPlacementInfo.h"
 #include "IPlacementModeModule.h"
+#include "Widgets/Input/SSegmentedControl.h"
 #include "Widgets/Views/STableViewBase.h"
 #include "Widgets/Views/STableRow.h"
 #include "Misc/TextFilter.h"
@@ -97,6 +98,7 @@ private:
 
 	/** When the tab is clicked we adjust the check state, so that the right style is displayed. */
 	void OnPlacementTabChanged( ECheckBoxState NewState, FName CategoryName );
+	void OnCategoryChanged(FName InCategory);
 
 	/** Gets the tab 'active' state, so that we can show the active style */
 	ECheckBoxState GetPlacementTabCheckedState( FName CategoryName ) const;
@@ -142,6 +144,12 @@ private:
 	/* The search box used to update the filter text */
 	TSharedPtr<SSearchBox> SearchBoxPtr;
 
+	/* Category Filter */
+	TSharedPtr<SSegmentedControl<FName>> CategoryFilterPtr;
+
+	/* Active Category Filter Label */
+	TSharedPtr<STextBlock> FilterLabelPtr;
+
 	/** Array of filtered items to show in the list view */
 	TArray<TSharedPtr<FPlaceableItem>> FilteredItems;
 
diff --git a/Engine/Source/Editor/PlacementMode/Public/IPlacementModeModule.h b/Engine/Source/Editor/PlacementMode/Public/IPlacementModeModule.h
--- a/Engine/Source/Editor/PlacementMode/Public/IPlacementModeModule.h
+++ b/Engine/Source/Editor/PlacementMode/Public/IPlacementModeModule.h
@@ -42,14 +42,17 @@ private:
  */
 struct FPlacementCategoryInfo
 {
-	FPlacementCategoryInfo(FText InDisplayName, FName InHandle, FString InTag, int32 InSortOrder = 0, bool bInSortable = true)
-		: DisplayName(InDisplayName), UniqueHandle(InHandle), SortOrder(InSortOrder), TagMetaData(MoveTemp(InTag)), bSortable(bInSortable)
+	FPlacementCategoryInfo(FText InDisplayName, FSlateIcon InDisplayIcon, FName InHandle, FString InTag, int32 InSortOrder = 0, bool bInSortable = true)
+		: DisplayName(InDisplayName), DisplayIcon(InDisplayIcon), UniqueHandle(InHandle), SortOrder(InSortOrder), TagMetaData(MoveTemp(InTag)), bSortable(bInSortable)
 	{
 	}
 
 	/** This category's display name */
 	FText DisplayName;
 
+	/** This category's representative icon */
+	FSlateIcon DisplayIcon;
+
 	/** A unique name for this category */
 	FName UniqueHandle;
 
