Fix placement mode compile error

diff --git a/Engine/Source/Editor/PlacementMode/Public/IPlacementModeModule.h b/Engine/Source/Editor/PlacementMode/Public/IPlacementModeModule.h
--- a/Engine/Source/Editor/PlacementMode/Public/IPlacementModeModule.h
+++ b/Engine/Source/Editor/PlacementMode/Public/IPlacementModeModule.h
@@ -47,6 +47,12 @@ struct FPlacementCategoryInfo
 	{
 	}
 
+	FPlacementCategoryInfo(FText InDisplayName, FName InHandle, FString InTag, int32 InSortOrder = 0, bool bInSortable = true)
+		: DisplayName(InDisplayName), DisplayIcon(), UniqueHandle(InHandle), SortOrder(InSortOrder), TagMetaData(MoveTemp(InTag)), bSortable(bInSortable)
+	{
+	}
+
+
 	/** This category's display name */
 	FText DisplayName;
 
