Missing files for PlacementBrowser

diff --git a/Engine/Plugins/MovieScene/LevelSequenceEditor/Source/LevelSequenceEditor/Private/LevelSequenceEditorModule.cpp b/Engine/Plugins/MovieScene/LevelSequenceEditor/Source/LevelSequenceEditor/Private/LevelSequenceEditorModule.cpp
--- a/Engine/Plugins/MovieScene/LevelSequenceEditor/Source/LevelSequenceEditor/Private/LevelSequenceEditorModule.cpp
+++ b/Engine/Plugins/MovieScene/LevelSequenceEditor/Source/LevelSequenceEditor/Private/LevelSequenceEditorModule.cpp
@@ -175,6 +175,7 @@ protected:
 	{
 		FPlacementCategoryInfo Info(
 			LOCTEXT("CinematicCategoryName", "Cinematic"),
+			FSlateIcon(FAppStyle::GetAppStyleSetName(), "PlacementBrowser.Icons.Cinematics"),
 			"Cinematic",
 			TEXT("PMCinematic"),
 			25
diff --git a/Engine/Plugins/Tests/FunctionalTestingEditor/Source/Private/FunctionalTestingEditorModule.cpp b/Engine/Plugins/Tests/FunctionalTestingEditor/Source/Private/FunctionalTestingEditorModule.cpp
--- a/Engine/Plugins/Tests/FunctionalTestingEditor/Source/Private/FunctionalTestingEditorModule.cpp
+++ b/Engine/Plugins/Tests/FunctionalTestingEditor/Source/Private/FunctionalTestingEditorModule.cpp
@@ -188,6 +188,7 @@ class FFunctionalTestingEditorModule : public IFunctionalTestingEditorModule
 		{
 			FPlacementCategoryInfo Info(
 				LOCTEXT("FunctionalTestingCategoryName", "Testing"),
+				FSlateIcon(FAppStyle::GetAppStyleSetName(), "PlacementBrowser.Icons.Testing"),
 				"Testing",
 				TEXT("PMTesting"),
 				25
diff --git a/Engine/Source/Editor/PlacementMode/Private/PlacementModeModule.cpp b/Engine/Source/Editor/PlacementMode/Private/PlacementModeModule.cpp
--- a/Engine/Source/Editor/PlacementMode/Private/PlacementModeModule.cpp
+++ b/Engine/Source/Editor/PlacementMode/Private/PlacementModeModule.cpp
@@ -101,6 +101,7 @@ void FPlacementModeModule::StartupModule()
 	RegisterPlacementCategory(
 		FPlacementCategoryInfo(
 			NSLOCTEXT("PlacementMode", "RecentlyPlaced", "Recently Placed"),
+			FSlateIcon(FAppStyle::GetAppStyleSetName(), "PlacementBrowser.Icons.Recent"),
 			FBuiltInPlacementCategories::RecentlyPlaced(),
 			TEXT("PMRecentlyPlaced"),
 			TNumericLimits<int32>::Lowest(),
@@ -114,6 +115,7 @@ void FPlacementModeModule::StartupModule()
 		RegisterPlacementCategory(
 			FPlacementCategoryInfo(
 				NSLOCTEXT("PlacementMode", "Basic", "Basic"),
+				FSlateIcon(FAppStyle::GetAppStyleSetName(), "PlacementBrowser.Icons.Basic"),
 				CategoryName,
 				TEXT("PMBasic"),
 				10
@@ -147,6 +149,7 @@ void FPlacementModeModule::StartupModule()
 		RegisterPlacementCategory(
 			FPlacementCategoryInfo(
 				NSLOCTEXT("PlacementMode", "Lights", "Lights"),
+				FSlateIcon(FAppStyle::GetAppStyleSetName(), "PlacementBrowser.Icons.Lights"),
 				CategoryName,
 				TEXT("PMLights"),
 				20
@@ -167,6 +170,7 @@ void FPlacementModeModule::StartupModule()
 		RegisterPlacementCategory(
 			FPlacementCategoryInfo(
 				NSLOCTEXT("PlacementMode", "VisualEffects", "Visual Effects"),
+				FSlateIcon(FAppStyle::GetAppStyleSetName(), "PlacementBrowser.Icons.VisualEffects"),
 				CategoryName,
 				TEXT("PMVisual"),
 				30
@@ -190,6 +194,7 @@ void FPlacementModeModule::StartupModule()
 	RegisterPlacementCategory(
 		FPlacementCategoryInfo(
 			NSLOCTEXT("PlacementMode", "Volumes", "Volumes"),
+			FSlateIcon(FAppStyle::GetAppStyleSetName(), "PlacementBrowser.Icons.Volumes"),
 			FBuiltInPlacementCategories::Volumes(),
 			TEXT("PMVolumes"),
 			40
@@ -199,6 +204,7 @@ void FPlacementModeModule::StartupModule()
 	RegisterPlacementCategory(
 		FPlacementCategoryInfo(
 			NSLOCTEXT("PlacementMode", "AllClasses", "All Classes"),
+			FSlateIcon(FAppStyle::GetAppStyleSetName(), "PlacementBrowser.Icons.All"),
 			FBuiltInPlacementCategories::AllClasses(),
 			TEXT("PMAllClasses"),
 			50
