Remove legacy placement editor mode.

#rnx

diff --git a/Engine/Plugins/Experimental/ControlRig/Source/ControlRigEditor/Private/EditMode/ControlRigEditMode.cpp b/Engine/Plugins/Experimental/ControlRig/Source/ControlRigEditor/Private/EditMode/ControlRigEditMode.cpp
--- a/Engine/Plugins/Experimental/ControlRig/Source/ControlRigEditor/Private/EditMode/ControlRigEditMode.cpp
+++ b/Engine/Plugins/Experimental/ControlRig/Source/ControlRigEditor/Private/EditMode/ControlRigEditMode.cpp
@@ -904,10 +904,6 @@ bool FControlRigEditMode::ShouldDrawWidget() const
 
 bool FControlRigEditMode::IsCompatibleWith(FEditorModeID OtherModeID) const
 {
-	if (OtherModeID == FBuiltinEditorModes::EM_Placement)
-	{
-		return false;
-	}
 	return true;
 }
 
diff --git a/Engine/Source/Editor/LevelEditor/Private/LevelEditor.cpp b/Engine/Source/Editor/LevelEditor/Private/LevelEditor.cpp
--- a/Engine/Source/Editor/LevelEditor/Private/LevelEditor.cpp
+++ b/Engine/Source/Editor/LevelEditor/Private/LevelEditor.cpp
@@ -207,12 +207,6 @@ TSharedRef<SDockTab> FLevelEditorModule::SpawnLevelEditor( const FSpawnTabArgs&
 		SetLevelEditorInstance(LevelEditorTmp);
 		LevelEditorTmp->Initialize(LevelEditorTab, OwnerWindow.ToSharedRef());
 
-		if (GetDefault<UEditorStyleSettings>()->bEnableLegacyEditorModeUI)
-		{
-			GLevelEditorModeTools().RemoveDefaultMode(FBuiltinEditorModes::EM_Default);
-			GLevelEditorModeTools().AddDefaultMode(FBuiltinEditorModes::EM_Placement);
-		}
-
 		GLevelEditorModeTools().DeactivateAllModes();
 		GLevelEditorModeTools().ActivateDefaultMode();
 
@@ -221,15 +215,6 @@ TSharedRef<SDockTab> FLevelEditorModule::SpawnLevelEditor( const FSpawnTabArgs&
 			// In legacy mode this toolbox should always be open
 			static const FTabId ToolboxTabId("LevelEditorToolBox");
 			LevelEditorTabManager->TryInvokeTab(ToolboxTabId);
-
-			// In legacy mode the standalone placement browser tab should not be opened
-			static const FTabId PlacementBrowserTabId("PlacementBrowser");
-			TSharedPtr<SDockTab> PlacementBrowserTab = LevelEditorTabManager->FindExistingLiveTab(PlacementBrowserTabId);
-			if (PlacementBrowserTab.IsValid())
-			{
-				PlacementBrowserTab->RequestCloseTab();
-
-			}
 		}
 
 		LevelEditorCreatedEvent.Broadcast(LevelEditorTmp);
diff --git a/Engine/Source/Editor/LevelEditor/Private/LevelEditorActions.cpp b/Engine/Source/Editor/LevelEditor/Private/LevelEditorActions.cpp
--- a/Engine/Source/Editor/LevelEditor/Private/LevelEditorActions.cpp
+++ b/Engine/Source/Editor/LevelEditor/Private/LevelEditorActions.cpp
@@ -56,7 +56,6 @@
 #include "SnappingUtils.h"
 #include "LevelEditorViewport.h"
 #include "Layers/LayersSubsystem.h"
-#include "IPlacementMode.h"
 #include "IPlacementModeModule.h"
 #include "AssetSelection.h"
 #include "IDocumentation.h"
diff --git a/Engine/Source/Editor/LevelEditor/Private/SLevelEditor.cpp b/Engine/Source/Editor/LevelEditor/Private/SLevelEditor.cpp
--- a/Engine/Source/Editor/LevelEditor/Private/SLevelEditor.cpp
+++ b/Engine/Source/Editor/LevelEditor/Private/SLevelEditor.cpp
@@ -1456,14 +1456,6 @@ void SLevelEditor::ToggleEditorMode( FEditorModeID ModeID )
 
 bool SLevelEditor::IsModeActive( FEditorModeID ModeID )
 {
-	// The level editor changes the default mode to placement
-	if ( ModeID == FBuiltinEditorModes::EM_Placement )
-	{
-		if (!GLevelEditorModeTools().IsOnlyVisibleActiveMode(ModeID))
-		{
-			return false;
-		}
-	}
 	return GLevelEditorModeTools().IsModeActive( ModeID );
 }
 
@@ -1535,7 +1527,7 @@ void SLevelEditor::RefreshEditorModeCommands()
 				Commands.EditorModeCommands[CommandIndex],
 				FExecuteAction::CreateStatic( &SLevelEditor::ToggleEditorMode, Mode.ID ),
 				FCanExecuteAction(),
-				FIsActionChecked::CreateStatic( &SLevelEditor::IsModeActive, Mode.ID ));
+				FIsActionChecked::CreateSP( GetEditorModeManager().AsShared(), &FEditorModeTools::IsModeActive, Mode.ID ));
 		}
 
 		CommandIndex++;
diff --git a/Engine/Source/Editor/LevelEditor/Private/SLevelEditor.h b/Engine/Source/Editor/LevelEditor/Private/SLevelEditor.h
--- a/Engine/Source/Editor/LevelEditor/Private/SLevelEditor.h
+++ b/Engine/Source/Editor/LevelEditor/Private/SLevelEditor.h
@@ -181,6 +181,7 @@ private:
 	static void ToggleEditorMode( FEditorModeID ModeID );
 
 	/** Checks if the editor mode is active for the auto-generated editor mode command. */
+	UE_DEPRECATED(5.0, "Use GLevelEditorModeTools().IsModeActive.")
 	static bool IsModeActive( FEditorModeID ModeID );
 
 	/**
diff --git a/Engine/Source/Editor/PlacementMode/Private/PlacementMode.cpp b/Engine/Source/Editor/PlacementMode/Private/PlacementMode.cpp
deleted file mode 100644
diff --git a/Engine/Source/Editor/PlacementMode/Private/PlacementMode.h b/Engine/Source/Editor/PlacementMode/Private/PlacementMode.h
deleted file mode 100644
diff --git a/Engine/Source/Editor/PlacementMode/Private/PlacementModeModule.cpp b/Engine/Source/Editor/PlacementMode/Private/PlacementModeModule.cpp
--- a/Engine/Source/Editor/PlacementMode/Private/PlacementModeModule.cpp
+++ b/Engine/Source/Editor/PlacementMode/Private/PlacementModeModule.cpp
@@ -45,7 +45,6 @@
 #include "AssetRegistryModule.h"
 #include "ActorPlacementInfo.h"
 #include "IPlacementModeModule.h"
-#include "PlacementMode.h"
 #include "AssetToolsModule.h"
 #include "Kismet2/KismetEditorUtilities.h"
 #include "ActorFactories/ActorFactoryPlanarReflection.h"
@@ -84,12 +83,6 @@ void FPlacementModeModule::StartupModule()
 		RecentlyPlaced.Add(FActorPlacementInfo(RecentlyPlacedAsStrings[Index]));
 	}
 
-	FEditorModeRegistry::Get().RegisterMode<FPlacementMode>(
-		FBuiltinEditorModes::EM_Placement,
-		NSLOCTEXT("PlacementMode", "DisplayName", "Place"),
-		FSlateIcon(FEditorStyle::GetStyleSetName(), "LevelEditor.PlacementMode", "LevelEditor.PlacementMode.Small"),
-		GetDefault<UEditorStyleSettings>()->bEnableLegacyEditorModeUI, 0);
-
 	FAssetRegistryModule& AssetRegistryModule = FModuleManager::LoadModuleChecked<FAssetRegistryModule>(TEXT("AssetRegistry"));
 	AssetRegistryModule.Get().OnAssetRemoved().AddRaw(this, &FPlacementModeModule::OnAssetRemoved);
 	AssetRegistryModule.Get().OnAssetRenamed().AddRaw(this, &FPlacementModeModule::OnAssetRenamed);
@@ -221,8 +214,6 @@ void FPlacementModeModule::StartupModule()
 
 void FPlacementModeModule::PreUnloadCallback()
 {
-	FEditorModeRegistry::Get().UnregisterMode(FBuiltinEditorModes::EM_Placement);
-
 	FAssetRegistryModule* AssetRegistryModule = FModuleManager::GetModulePtr<FAssetRegistryModule>("AssetRegistry");
 	if (AssetRegistryModule)
 	{
diff --git a/Engine/Source/Editor/PlacementMode/Private/PlacementModeToolkit.h b/Engine/Source/Editor/PlacementMode/Private/PlacementModeToolkit.h
deleted file mode 100644
diff --git a/Engine/Source/Editor/PlacementMode/Private/SPlacementModeTools.cpp b/Engine/Source/Editor/PlacementMode/Private/SPlacementModeTools.cpp
--- a/Engine/Source/Editor/PlacementMode/Private/SPlacementModeTools.cpp
+++ b/Engine/Source/Editor/PlacementMode/Private/SPlacementModeTools.cpp
@@ -13,7 +13,6 @@
 #include "EditorModes.h"
 #include "AssetThumbnail.h"
 #include "LevelEditor.h"
-#include "PlacementMode.h"
 #include "DragAndDrop/AssetDragDropOp.h"
 #include "EditorClassUtils.h"
 #include "Widgets/Input/SSearchBox.h"
@@ -319,13 +318,6 @@ void SPlacementModeTools::Construct( const FArguments& InArgs )
 	static const FName PlacementBrowserActiveTabBarBrushName("PlacementBrowser.ActiveTabBar");
 	PlacementBrowserActiveTabBarBrush = FEditorStyle::GetBrush(PlacementBrowserActiveTabBarBrushName);
 
-
-	FPlacementMode* PlacementEditMode = (FPlacementMode*)GLevelEditorModeTools().GetActiveMode( FBuiltinEditorModes::EM_Placement );
-	if (PlacementEditMode)
-	{
-		PlacementEditMode->AddValidFocusTargetForPlacement(SharedThis(this));
-	}
-
 	SearchTextFilter = MakeShareable(new FPlacementAssetEntryTextFilter(
 		FPlacementAssetEntryTextFilter::FItemToStringArray::CreateStatic(&PlacementViewFilter::GetBasicStrings)
 		));
@@ -644,24 +636,6 @@ void SPlacementModeTools::Tick( const FGeometry& AllottedGeometry, const double
 	}
 }
 
-FReply SPlacementModeTools::OnKeyDown( const FGeometry& MyGeometry, const FKeyEvent& InKeyEvent )
-{
-	FReply Reply = FReply::Unhandled();
-
-	if ( InKeyEvent.GetKey() == EKeys::Escape )
-	{
-		FPlacementMode* PlacementEditMode = (FPlacementMode*)GLevelEditorModeTools().GetActiveMode( FBuiltinEditorModes::EM_Placement );
-		// Catch potential nullptr
-		if (ensureMsgf(PlacementEditMode, TEXT("PlacementEditMode was null, but SPlacementModeTools is still accepting KeyDown events")))
-		{
-			PlacementEditMode->StopPlacing();
-		}
-		Reply = FReply::Handled();
-	}
-
-	return Reply;
-}
-
 void SPlacementModeTools::OnSearchChanged(const FText& InFilterText)
 {
 	// If the search text was previously empty we do a full rebuild of our cached widgets
diff --git a/Engine/Source/Editor/PlacementMode/Private/SPlacementModeTools.h b/Engine/Source/Editor/PlacementMode/Private/SPlacementModeTools.h
--- a/Engine/Source/Editor/PlacementMode/Private/SPlacementModeTools.h
+++ b/Engine/Source/Editor/PlacementMode/Private/SPlacementModeTools.h
@@ -65,7 +65,6 @@ private:
 
 	// Begin SWidget
 	virtual void Tick( const FGeometry& AllottedGeometry, const double InCurrentTime, const float InDeltaTime ) override;
-	virtual FReply OnKeyDown( const FGeometry& MyGeometry, const FKeyEvent& InKeyEvent ) override;
 	// End SWidget
 
 private:
diff --git a/Engine/Source/Editor/PlacementMode/Public/IPlacementMode.h b/Engine/Source/Editor/PlacementMode/Public/IPlacementMode.h
deleted file mode 100644
diff --git a/Engine/Source/Editor/PlacementMode/Public/IPlacementModeModule.h b/Engine/Source/Editor/PlacementMode/Public/IPlacementModeModule.h
--- a/Engine/Source/Editor/PlacementMode/Public/IPlacementModeModule.h
+++ b/Engine/Source/Editor/PlacementMode/Public/IPlacementModeModule.h
@@ -11,7 +11,6 @@
 #include "Modules/ModuleManager.h"
 #include "ActorFactories/ActorFactory.h"
 #include "ActorPlacementInfo.h"
-#include "IPlacementMode.h"
 #include "GameFramework/Volume.h"
 #include "Editor.h"
 
@@ -230,35 +229,21 @@ public:
 	/**
 	 * @return the event that is broadcast whenever a placement mode category is refreshed
 	 */
-	DECLARE_EVENT_OneParam( IPlacementMode, FOnPlacementModeCategoryRefreshed, FName /*CategoryName*/ );
+	DECLARE_EVENT_OneParam( IPlacementModeModule, FOnPlacementModeCategoryRefreshed, FName /*CategoryName*/ );
 	virtual FOnPlacementModeCategoryRefreshed& OnPlacementModeCategoryRefreshed() = 0;
 
 	/**
 	 * @return the event that is broadcast whenever the list of recently placed assets changes
 	 */
-	DECLARE_EVENT_OneParam( IPlacementMode, FOnRecentlyPlacedChanged, const TArray< FActorPlacementInfo >& /*NewRecentlyPlaced*/ );
+	DECLARE_EVENT_OneParam( IPlacementModeModule, FOnRecentlyPlacedChanged, const TArray< FActorPlacementInfo >& /*NewRecentlyPlaced*/ );
 	virtual FOnRecentlyPlacedChanged& OnRecentlyPlacedChanged() = 0;
 
 	/**
 	 * @return the event that is broadcast whenever the list of all placeable assets changes
 	 */
-	DECLARE_EVENT( IPlacementMode, FOnAllPlaceableAssetsChanged );
+	DECLARE_EVENT( IPlacementModeModule, FOnAllPlaceableAssetsChanged );
 	virtual FOnAllPlaceableAssetsChanged& OnAllPlaceableAssetsChanged() = 0;
 
-	/**
-	 * @return the event that is broadcast whenever a placement mode enters a placing session
-	 */
-	DECLARE_EVENT_OneParam( IPlacementMode, FOnStartedPlacingEvent, const TArray< UObject* >& /*Assets*/ );
-	virtual FOnStartedPlacingEvent& OnStartedPlacing() = 0;
-	virtual void BroadcastStartedPlacing( const TArray< UObject* >& Assets ) = 0;
-
-	/**
-	 * @return the event that is broadcast whenever a placement mode exits a placing session
-	 */
-	DECLARE_EVENT_OneParam( IPlacementMode, FOnStoppedPlacingEvent, bool /*bWasSuccessfullyPlaced*/ );
-	virtual FOnStoppedPlacingEvent& OnStoppedPlacing() = 0;
-	virtual void BroadcastStoppedPlacing( bool bWasSuccessfullyPlaced ) = 0;
-
 	/**
 	 * Creates the placement browser widget
 	 */
diff --git a/Engine/Source/Editor/PlacementMode/Public/PlacementModeModule.h b/Engine/Source/Editor/PlacementMode/Public/PlacementModeModule.h
--- a/Engine/Source/Editor/PlacementMode/Public/PlacementModeModule.h
+++ b/Engine/Source/Editor/PlacementMode/Public/PlacementModeModule.h
@@ -79,28 +79,6 @@ public:
 		return RecentlyPlaced;
 	}
 
-	/** @return the event that is broadcast whenever the placement mode enters a placing session */
-	DECLARE_DERIVED_EVENT(FPlacementModeModule, IPlacementModeModule::FOnStartedPlacingEvent, FOnStartedPlacingEvent);
-	virtual FOnStartedPlacingEvent& OnStartedPlacing() override
-	{
-		return StartedPlacingEvent;
-	}
-	virtual void BroadcastStartedPlacing(const TArray< UObject* >& Assets) override
-	{
-		StartedPlacingEvent.Broadcast(Assets);
-	}
-
-	/** @return the event that is broadcast whenever the placement mode exits a placing session */
-	DECLARE_DERIVED_EVENT(FPlacementModeModule, IPlacementModeModule::FOnStoppedPlacingEvent, FOnStoppedPlacingEvent);
-	virtual FOnStoppedPlacingEvent& OnStoppedPlacing() override
-	{
-		return StoppedPlacingEvent;
-	}
-	virtual void BroadcastStoppedPlacing(bool bWasSuccessfullyPlaced) override
-	{
-		StoppedPlacingEvent.Broadcast(bWasSuccessfullyPlaced);
-	}
-
 	virtual TSharedRef<SWidget> CreatePlacementModeBrowser() override;
 
 	virtual bool RegisterPlacementCategory(const FPlacementCategoryInfo& Info);
@@ -146,9 +124,6 @@ private:
 	FOnAllPlaceableAssetsChanged	AllPlaceableAssetsChanged;
 	FOnPlacementModeCategoryRefreshed PlacementModeCategoryRefreshed;
 
-	FOnStartedPlacingEvent			StartedPlacingEvent;
-	FOnStoppedPlacingEvent			StoppedPlacingEvent;
-
 	TArray< TSharedPtr<FExtender> > ContentPaletteFiltersExtenders;
 	TArray< TSharedPtr<FExtender> > PaletteExtenders;
 
diff --git a/Engine/Source/Editor/UnrealEd/Private/EditorModeInterpolation.cpp b/Engine/Source/Editor/UnrealEd/Private/EditorModeInterpolation.cpp
--- a/Engine/Source/Editor/UnrealEd/Private/EditorModeInterpolation.cpp
+++ b/Engine/Source/Editor/UnrealEd/Private/EditorModeInterpolation.cpp
@@ -190,9 +190,7 @@ void FEdModeInterpEdit::ActorSelectionChangeNotify()
 
 bool FEdModeInterpEdit::IsCompatibleWith(FEditorModeID OtherModeID) const
 {
-	return
-		OtherModeID == FBuiltinEditorModes::EM_Placement ||
-		OtherModeID == FBuiltinEditorModes::EM_MeshPaint;
+	return OtherModeID == FBuiltinEditorModes::EM_MeshPaint;
 }
 
 //////////////////////////////////////////////////////////////////////////
diff --git a/Engine/Source/Editor/UnrealEd/Private/EditorModeRegistry.cpp b/Engine/Source/Editor/UnrealEd/Private/EditorModeRegistry.cpp
--- a/Engine/Source/Editor/UnrealEd/Private/EditorModeRegistry.cpp
+++ b/Engine/Source/Editor/UnrealEd/Private/EditorModeRegistry.cpp
@@ -43,30 +43,21 @@ TSharedRef<FEdMode> FEditorModeFactory::CreateMode() const
 
 void FEditorModeRegistry::Initialize()
 {
-	// Send notifications for any legacy modes that were registered before the asset subsytem started up
+	// Send notifications for any legacy modes that were registered before the asset subsystem started up
 	UAssetEditorSubsystem* AssetEditorSubsystem = GEditor->GetEditorSubsystem<UAssetEditorSubsystem>();
 	for (const FactoryMap::ElementType& ModeEntry : ModeFactories)
 	{
 		AssetEditorSubsystem->OnEditorModeRegistered().Broadcast(ModeEntry.Key);
 	}
 
-	if(!GetDefault<UEditorStyleSettings>()->bEnableLegacyEditorModeUI)
-	{
-		// Add default editor modes
-		RegisterMode<FEdModeDefault>(
-			FBuiltinEditorModes::EM_Default,
-			NSLOCTEXT("DefaultMode", "DisplayName", "Select"),
-			FSlateIcon(FEditorStyle::GetStyleSetName(), "LevelEditor.SelectMode", "LevelEditor.SelectMode.Small"),
-			true, 0);
-	}
-	else
-	{
-		RegisterMode<FEdModeDefault>(FBuiltinEditorModes::EM_Default);
-	}
-
+	// Add default editor modes
+	RegisterMode<FEdModeDefault>(
+		FBuiltinEditorModes::EM_Default,
+		NSLOCTEXT("DefaultMode", "DisplayName", "Select"),
+		FSlateIcon(FEditorStyle::GetStyleSetName(), "LevelEditor.SelectMode", "LevelEditor.SelectMode.Small"),
+		true, 0);
 	RegisterMode<FEdModeInterpEdit>(FBuiltinEditorModes::EM_InterpEdit);
 
-	FModuleManager::LoadModuleChecked<IPlacementModeModule>(TEXT("PlacementMode"));
 	FModuleManager::LoadModuleChecked<FActorPickerModeModule>(TEXT("ActorPickerMode"));
 	FModuleManager::LoadModuleChecked<FSceneDepthPickerModeModule>(TEXT("SceneDepthPickerMode"));
 	FModuleManager::LoadModuleChecked<IMeshPaintModule>(TEXT("MeshPaintMode"));
diff --git a/Engine/Source/Editor/VREditor/UI/VREditorUISystem.cpp b/Engine/Source/Editor/VREditor/UI/VREditorUISystem.cpp
--- a/Engine/Source/Editor/VREditor/UI/VREditorUISystem.cpp
+++ b/Engine/Source/Editor/VREditor/UI/VREditorUISystem.cpp
@@ -1867,7 +1867,7 @@ void UVREditorUISystem::ResetAll()
 			ShowEditorUIPanel(FloatingUI, UIInteractor, bShouldShow, bSpawnInFront, bDragFromOpen, bPlaySound);
 		}
 	}
-	FVREditorActionCallbacks::ChangeEditorModes(FBuiltinEditorModes::EM_Placement);
+	FVREditorActionCallbacks::ChangeEditorModes(FBuiltinEditorModes::EM_Default);
 
 	// Turn off all snapping
 	if (FVREditorActionCallbacks::GetTranslationSnapState() == ECheckBoxState::Checked)
diff --git a/Engine/Source/Editor/VREditor/VREditorMode.cpp b/Engine/Source/Editor/VREditor/VREditorMode.cpp
--- a/Engine/Source/Editor/VREditor/VREditorMode.cpp
+++ b/Engine/Source/Editor/VREditor/VREditorMode.cpp
@@ -275,8 +275,8 @@ void UVREditorMode::Enter()
 		}
 	}
 
-	// Switch us back to placement mode and close any open sequencer windows
-	FVREditorActionCallbacks::ChangeEditorModes(FBuiltinEditorModes::EM_Placement);
+	// Switch us back to default mode and close any open sequencer windows
+	FVREditorActionCallbacks::ChangeEditorModes(FBuiltinEditorModes::EM_Default);
 	FLevelEditorModule& LevelEditorModule = FModuleManager::GetModuleChecked<FLevelEditorModule>(TEXT("LevelEditor"));
 	LevelEditorModule.GetLevelEditorTabManager()->TryInvokeTab(FTabId("Sequencer"))->RequestCloseTab();
 
@@ -363,7 +363,7 @@ void UVREditorMode::Exit(const bool bShouldDisableStereo)
 		GetLevelViewportPossessedForVR().RemoveAllPreviews(false);
 		GEditor->SelectNone(true, true, false);
 		GEditor->NoteSelectionChange();
-		FVREditorActionCallbacks::ChangeEditorModes(FBuiltinEditorModes::EM_Placement);
+		FVREditorActionCallbacks::ChangeEditorModes(FBuiltinEditorModes::EM_Default);
 		
 		//Destroy the avatar
 		{
