Placement Browser - Wrapped Category Filters

diff --git a/Engine/Source/Editor/EditorStyle/Private/StarshipStyle.cpp b/Engine/Source/Editor/EditorStyle/Private/StarshipStyle.cpp
--- a/Engine/Source/Editor/EditorStyle/Private/StarshipStyle.cpp
+++ b/Engine/Source/Editor/EditorStyle/Private/StarshipStyle.cpp
@@ -1793,12 +1793,12 @@ void FStarshipEditorStyle::FStyle::SetupGeneralStyles()
 
 		Set( "PaletteToolBar.Tab",  FCheckBoxStyle()
 			.SetCheckBoxType(            ESlateCheckBoxType::ToggleButton)
-			.SetCheckedImage(            FSlateRoundedBoxBrush(FStyleColors::Recessed, 2.0))
-			.SetCheckedHoveredImage(     FSlateRoundedBoxBrush(FStyleColors::Recessed, 2.0))
-			.SetCheckedPressedImage(     FSlateRoundedBoxBrush(FStyleColors::Recessed, 2.0))
-			.SetUncheckedImage(          FSlateRoundedBoxBrush(FStyleColors::Secondary, 2.0))
-			.SetUncheckedHoveredImage(   FSlateRoundedBoxBrush(FStyleColors::Secondary, 2.0))
-			.SetUncheckedPressedImage(   FSlateRoundedBoxBrush(FStyleColors::Secondary, 2.0))
+			.SetCheckedImage(            FSlateNoResource())
+			.SetCheckedHoveredImage(     FSlateNoResource())
+			.SetCheckedPressedImage(     FSlateNoResource())
+			.SetUncheckedImage(          FSlateNoResource())
+			.SetUncheckedHoveredImage(   FSlateNoResource())
+			.SetUncheckedPressedImage(   FSlateNoResource())
 			.SetForegroundColor(         FStyleColors::Foreground)
 			.SetHoveredForegroundColor(  FStyleColors::ForegroundHover)
 			.SetPressedForegroundColor(  FStyleColors::ForegroundHover)
diff --git a/Engine/Source/Editor/PlacementMode/Private/SPlacementModeTools.cpp b/Engine/Source/Editor/PlacementMode/Private/SPlacementModeTools.cpp
--- a/Engine/Source/Editor/PlacementMode/Private/SPlacementModeTools.cpp
+++ b/Engine/Source/Editor/PlacementMode/Private/SPlacementModeTools.cpp
@@ -322,10 +322,9 @@ void SPlacementModeTools::Construct( const FArguments& InArgs )
 		FPlacementAssetEntryTextFilter::FItemToStringArray::CreateStatic(&PlacementViewFilter::GetBasicStrings)
 		));
 
-	SAssignNew(CategoryFilterPtr, SSegmentedControl<FName>)
-	.Style(&FAppStyle::Get(), "PlacementBrowser.CategoryControl")
-    .OnValueChanged(this, &SPlacementModeTools::OnCategoryChanged)
-    .Value_Lambda( [this] { return ActiveTabName; } );
+	SAssignNew(CategoryFilterPtr, SUniformWrapPanel)
+	.HAlign(HAlign_Center)
+	.SlotPadding(FMargin(2.0f, 0.0f));
 
 	// Populate the categories and body from the defined placeable items
 	IPlacementModeModule& PlacementModeModule = IPlacementModeModule::Get();
@@ -334,7 +333,20 @@ void SPlacementModeTools::Construct( const FArguments& InArgs )
 	PlacementModeModule.GetSortedCategories(Categories);
 	for (const FPlacementCategoryInfo& Category : Categories)
 	{
-		CategoryFilterPtr->AddSlot(Category.UniqueHandle).Icon(Category.DisplayIcon.GetIcon());// .Text(Category.DisplayName);
+		CategoryFilterPtr->AddSlot()
+		[
+			SNew(SCheckBox)
+			.Padding(FMargin(4.f, 4.f))
+			.Style( &FAppStyle::Get(),  "PaletteToolBar.Tab" )
+			.OnCheckStateChanged(this, &SPlacementModeTools::OnCategoryChanged, Category.UniqueHandle)
+			.IsChecked(this, &SPlacementModeTools::GetPlacementTabCheckedState, Category.UniqueHandle)
+			.ToolTipText(Category.DisplayName)
+			[
+				  SNew(SImage)
+	              .ColorAndOpacity(FSlateColor::UseForeground())
+	              .Image(Category.DisplayIcon.GetIcon())
+			]
+		];
 	}
 
 	TSharedRef<SScrollBar> ScrollBar = SNew(SScrollBar)
@@ -360,11 +372,12 @@ void SPlacementModeTools::Construct( const FArguments& InArgs )
 
 		+ SVerticalBox::Slot()
 		.AutoHeight()
+		.HAlign(HAlign_Fill)
 		[
 			SNew(SBorder)
 			.BorderImage(FEditorStyle::GetBrush("ToolPanel.GroupBorder"))
 			.Padding(FMargin(8.f, 6.f, 8.f, 8.f))
-			.HAlign(HAlign_Center)
+			.HAlign(HAlign_Fill)
 			.Visibility(this, &SPlacementModeTools::GetTabsVisibility)
 			[
 				CategoryFilterPtr.ToSharedRef()
@@ -564,9 +577,9 @@ TSharedRef<ITableRow> SPlacementModeTools::OnGenerateWidgetForItem(TSharedPtr<FP
 		];
 }
 
-void SPlacementModeTools::OnCategoryChanged(FName InCategory)
+void SPlacementModeTools::OnCategoryChanged(const ECheckBoxState NewState, FName InCategory)
 {
-	if (InCategory != ActiveTabName)
+	if (NewState == ECheckBoxState::Checked && InCategory != ActiveTabName)
 	{
 		ActiveTabName = InCategory;
 		IPlacementModeModule::Get().RegenerateItemsForCategory(ActiveTabName);
diff --git a/Engine/Source/Editor/PlacementMode/Private/SPlacementModeTools.h b/Engine/Source/Editor/PlacementMode/Private/SPlacementModeTools.h
--- a/Engine/Source/Editor/PlacementMode/Private/SPlacementModeTools.h
+++ b/Engine/Source/Editor/PlacementMode/Private/SPlacementModeTools.h
@@ -9,7 +9,7 @@
 #include "Widgets/SCompoundWidget.h"
 #include "ActorPlacementInfo.h"
 #include "IPlacementModeModule.h"
-#include "Widgets/Input/SSegmentedControl.h"
+#include "Widgets/Layout/SUniformWrapPanel.h"
 #include "Widgets/Views/STableViewBase.h"
 #include "Widgets/Views/STableRow.h"
 #include "Misc/TextFilter.h"
@@ -97,7 +97,7 @@ private:
 
 	/** When the tab is clicked we adjust the check state, so that the right style is displayed. */
 	void OnPlacementTabChanged( ECheckBoxState NewState, FName CategoryName );
-	void OnCategoryChanged(FName InCategory);
+	void OnCategoryChanged(const ECheckBoxState NewState, FName InCategory);
 
 	/** Gets the tab 'active' state, so that we can show the active style */
 	ECheckBoxState GetPlacementTabCheckedState( FName CategoryName ) const;
@@ -144,7 +144,7 @@ private:
 	TSharedPtr<SSearchBox> SearchBoxPtr;
 
 	/* Category Filter */
-	TSharedPtr<SSegmentedControl<FName>> CategoryFilterPtr;
+	TSharedPtr<SUniformWrapPanel> CategoryFilterPtr;
 
 	/* Active Category Filter Label */
 	TSharedPtr<STextBlock> FilterLabelPtr;
