Place Actors (aka Placement Mode) Quick Toolbar Menu

diff --git a/Engine/Source/Editor/EditorStyle/Private/StarshipStyle.cpp b/Engine/Source/Editor/EditorStyle/Private/StarshipStyle.cpp
--- a/Engine/Source/Editor/EditorStyle/Private/StarshipStyle.cpp
+++ b/Engine/Source/Editor/EditorStyle/Private/StarshipStyle.cpp
@@ -4566,6 +4566,7 @@ void FStarshipEditorStyle::FStyle::SetupLevelEditorStyle()
 
 		Set( "LevelEditor.WorldProperties.Tab", new IMAGE_BRUSH_SVG( "Starship/Common/WorldSettings", Icon16x16 ) );
 
+		Set( "LevelEditor.OpenCreateMenu", new IMAGE_BRUSH_SVG( "Starship/Common/PlaceActors", Icon20x20 ) );
 		Set( "LevelEditor.OpenContentBrowser", new IMAGE_BRUSH_SVG( "Starship/MainToolbar/content", Icon20x20 ) );
 		Set( "LevelEditor.OpenMarketplace", new IMAGE_BRUSH_SVG( "Starship/MainToolbar/marketplace", Icon20x20) );
 		Set( "LevelEditor.OpenLevelBlueprint", new IMAGE_BRUSH_SVG( "Starship/MainToolbar/blueprints", Icon20x20) );
diff --git a/Engine/Source/Editor/LevelEditor/Private/LevelEditorActions.cpp b/Engine/Source/Editor/LevelEditor/Private/LevelEditorActions.cpp
--- a/Engine/Source/Editor/LevelEditor/Private/LevelEditorActions.cpp
+++ b/Engine/Source/Editor/LevelEditor/Private/LevelEditorActions.cpp
@@ -2296,6 +2296,12 @@ void FLevelEditorActionCallbacks::OnShowWorldProperties( TWeakPtr< SLevelEditor
 	LevelEditorModule.GetLevelEditorTabManager()->TryInvokeTab(FName("WorldSettingsTab"));
 }
 
+void FLevelEditorActionCallbacks::OpenPlaceActors()
+{
+	FLevelEditorModule& LevelEditorModule = FModuleManager::GetModuleChecked<FLevelEditorModule>( TEXT("LevelEditor") );
+	LevelEditorModule.GetLevelEditorTabManager()->TryInvokeTab(LevelEditorTabIds::PlacementBrowser);
+}
+
 void FLevelEditorActionCallbacks::OpenContentBrowser()
 {
 	FContentBrowserModule& ContentBrowserModule = FModuleManager::LoadModuleChecked<FContentBrowserModule>("ContentBrowser");
@@ -3430,6 +3436,7 @@ void FLevelEditorCommands::RegisterCommands()
 	UI_COMMAND( FindActorInLevelScript, "Find in Level Blueprint", "Finds any references to the selected actor in its level's blueprint", EUserInterfaceActionType::Button, FInputChord() );
 
 	UI_COMMAND( WorldProperties, "World Settings", "Displays the world settings", EUserInterfaceActionType::Button, FInputChord() );
+	UI_COMMAND( OpenPlaceActors, "Place Actors Panel", "Opens the Place Actors Panel", EUserInterfaceActionType::Button, FInputChord(EModifierKey::Alt, EKeys::P) );
 	UI_COMMAND( OpenContentBrowser, "Open Content Browser", "Opens the Content Browser", EUserInterfaceActionType::Button, FInputChord(EModifierKey::Control|EModifierKey::Shift, EKeys::F) );
 	UI_COMMAND( OpenMarketplace, "Open Marketplace", "Opens the Marketplace", EUserInterfaceActionType::Button, FInputChord() );
 
diff --git a/Engine/Source/Editor/LevelEditor/Private/LevelEditorToolBar.cpp b/Engine/Source/Editor/LevelEditor/Private/LevelEditorToolBar.cpp
--- a/Engine/Source/Editor/LevelEditor/Private/LevelEditorToolBar.cpp
+++ b/Engine/Source/Editor/LevelEditor/Private/LevelEditorToolBar.cpp
@@ -1178,6 +1178,7 @@ void FLevelEditorToolBar::RegisterLevelEditorToolBar( const TSharedRef<FUIComman
 	RegisterQuickSettingsMenu();
 	RegisterOpenBlueprintMenu();
 	RegisterContentMenu();
+	RegisterCreateMenu();
 
 #define LOCTEXT_NAMESPACE "LevelEditorToolBar"
 
@@ -1200,6 +1201,16 @@ void FLevelEditorToolBar::RegisterLevelEditorToolBar( const TSharedRef<FUIComman
 
 		{
 			FToolMenuSection& Section = AssetsToolBar->AddSection("Content");
+
+			Section.AddEntry(FToolMenuEntry::InitComboButton(
+				"Create",
+				FUIAction(),
+				FOnGetContent::CreateStatic(&FLevelEditorToolBar::GenerateCreateMenuWidget, InCommandList, TWeakPtr<SLevelEditor>(InLevelEditor)),
+				LOCTEXT("OpenCreate_Label", "Content"),
+				LOCTEXT("OpenCreate_Tooltip", "Quickly add to the project."),
+				FSlateIcon(FEditorStyle::GetStyleSetName(), "LevelEditor.OpenCreateMenu")
+			));
+
 			Section.AddEntry(FToolMenuEntry::InitComboButton(
 				"Content",
 				FUIAction(),
@@ -1869,6 +1880,17 @@ void FLevelEditorToolBar::RegisterContentMenu()
 #undef LOCTEXT_NAMESPACE
 }
 
+void FLevelEditorToolBar::RegisterCreateMenu()
+{
+#define LOCTEXT_NAMESPACE "LevelToolBarViewMenu"
+	UToolMenu* ContentMenu = UToolMenus::Get()->RegisterMenu("LevelEditor.LevelEditorToolBar.CreateQuickMenu");
+	{
+		FToolMenuSection& Section = ContentMenu->FindOrAddSection("Create");
+		Section.InitSection("ProjectContent", LOCTEXT("Create_Label", "Create"), FToolMenuInsert());
+	}
+#undef LOCTEXT_NAMESPACE
+}
+
 void FLevelEditorToolBar::OnOpenSubLevelBlueprint( ULevel* InLevel )
 {
 	ULevelScriptBlueprint* LevelScriptBlueprint = InLevel->GetLevelScriptBlueprint();
@@ -1905,6 +1927,16 @@ TSharedRef< SWidget > FLevelEditorToolBar::GenerateContentMenuWidget(TSharedRef<
 	return UToolMenus::Get()->GenerateWidget("LevelEditor.LevelEditorToolBar.ContentQuickMenu", MenuContext);
 }
 
+TSharedRef< SWidget > FLevelEditorToolBar::GenerateCreateMenuWidget(TSharedRef<FUICommandList> InCommandList, TWeakPtr<SLevelEditor> InLevelEditor)
+{
+	FToolMenuContext MenuContext(InCommandList);
+	ULevelEditorMenuContext* LevelEditorMenuContext = NewObject<ULevelEditorMenuContext>();
+	LevelEditorMenuContext->LevelEditor = InLevelEditor;
+	MenuContext.AddObject(LevelEditorMenuContext);
+
+	return UToolMenus::Get()->GenerateWidget("LevelEditor.LevelEditorToolBar.CreateQuickMenu", MenuContext);
+}
+
 void FLevelEditorToolBar::RegisterCinematicsMenu()
 {
 #define LOCTEXT_NAMESPACE "LevelToolBarCinematicsMenu"
diff --git a/Engine/Source/Editor/LevelEditor/Private/LevelEditorToolBar.h b/Engine/Source/Editor/LevelEditor/Private/LevelEditorToolBar.h
--- a/Engine/Source/Editor/LevelEditor/Private/LevelEditorToolBar.h
+++ b/Engine/Source/Editor/LevelEditor/Private/LevelEditorToolBar.h
@@ -66,6 +66,8 @@ protected:
 
 	static TSharedRef< SWidget > GenerateContentMenuWidget(TSharedRef<FUICommandList> InCommandList, TWeakPtr<SLevelEditor> InLevelEditor);
 
+	static TSharedRef< SWidget > GenerateCreateMenuWidget(TSharedRef<FUICommandList> InCommandList, TWeakPtr<SLevelEditor> InLevelEditor);
+
 	/**
 	 * Delegate for actor selection within the Cinematics popup menu's SceneOutliner.
 	 * Opens the editor for the selected actor and dismisses all popup menus.
@@ -86,4 +88,5 @@ private:
 	static void RegisterQuickSettingsMenu();
 	static void RegisterOpenBlueprintMenu();
 	static void RegisterContentMenu();
+	static void RegisterCreateMenu();
 };
diff --git a/Engine/Source/Editor/LevelEditor/Private/SLevelEditor.cpp b/Engine/Source/Editor/LevelEditor/Private/SLevelEditor.cpp
--- a/Engine/Source/Editor/LevelEditor/Private/SLevelEditor.cpp
+++ b/Engine/Source/Editor/LevelEditor/Private/SLevelEditor.cpp
@@ -127,6 +127,10 @@ void SLevelEditor::BindCommands()
 		FExecuteAction::CreateStatic( &FLevelEditorActionCallbacks::ConvertSelectedActorsIntoBlueprintClass ),
 		FCanExecuteAction::CreateStatic( &FLevelEditorActionCallbacks::CanConvertSelectedActorsIntoBlueprintClass ) );
 
+	LevelEditorCommands->MapAction(
+		Actions.OpenPlaceActors,
+		FExecuteAction::CreateStatic( &FLevelEditorActionCallbacks::OpenPlaceActors) );
+
 	LevelEditorCommands->MapAction(
 		Actions.OpenContentBrowser,
 		FExecuteAction::CreateStatic( &FLevelEditorActionCallbacks::OpenContentBrowser ) );
diff --git a/Engine/Source/Editor/LevelEditor/Public/LevelEditorActions.h b/Engine/Source/Editor/LevelEditor/Public/LevelEditorActions.h
--- a/Engine/Source/Editor/LevelEditor/Public/LevelEditorActions.h
+++ b/Engine/Source/Editor/LevelEditor/Public/LevelEditorActions.h
@@ -541,6 +541,7 @@ public:
 	 */
 
 	TSharedPtr< FUICommandInfo > WorldProperties;
+	TSharedPtr< FUICommandInfo > OpenPlaceActors;
 	TSharedPtr< FUICommandInfo > OpenContentBrowser;
 	TSharedPtr< FUICommandInfo > OpenMarketplace;
 	TSharedPtr< FUICommandInfo > ToggleVR;
@@ -1147,6 +1148,9 @@ public:
 	/** Select the world info actor and show the properties */
 	static void OnShowWorldProperties( TWeakPtr< SLevelEditor > LevelEditor );
 
+	/** Open the Place Actors Panel */
+	static void OpenPlaceActors();
+
 	/** Open the Content Browser */
 	static void OpenContentBrowser();
 
diff --git a/Engine/Source/Editor/PlacementMode/PlacementMode.Build.cs b/Engine/Source/Editor/PlacementMode/PlacementMode.Build.cs
--- a/Engine/Source/Editor/PlacementMode/PlacementMode.Build.cs
+++ b/Engine/Source/Editor/PlacementMode/PlacementMode.Build.cs
@@ -23,6 +23,7 @@ public class PlacementMode : ModuleRules
                 "CollectionManager",
                 "LevelEditor",
                 "AssetTools",
+                "ToolMenus",
                 "EditorStyle"
             } 
         );
diff --git a/Engine/Source/Editor/PlacementMode/Private/PlacementModeModule.cpp b/Engine/Source/Editor/PlacementMode/Private/PlacementModeModule.cpp
--- a/Engine/Source/Editor/PlacementMode/Private/PlacementModeModule.cpp
+++ b/Engine/Source/Editor/PlacementMode/Private/PlacementModeModule.cpp
@@ -37,6 +37,7 @@
 #include "Engine/StaticMesh.h"
 #include "GameFramework/Volume.h"
 #include "Engine/PostProcessVolume.h"
+#include "LevelEditorActions.h"
 #include "AssetData.h"
 #include "EditorModeRegistry.h"
 #include "EditorModes.h"
@@ -45,6 +46,7 @@
 #include "AssetRegistryModule.h"
 #include "ActorPlacementInfo.h"
 #include "IPlacementModeModule.h"
+#include "ToolMenus.h"
 #include "AssetToolsModule.h"
 #include "Kismet2/KismetEditorUtilities.h"
 #include "ActorFactories/ActorFactoryPlanarReflection.h"
@@ -121,17 +123,7 @@ void FPlacementModeModule::StartupModule()
 		Category->Items.Add(CreateID(), MakeShareable(new FPlaceableItem(*UActorFactoryPawn::StaticClass(), SortOrder += 10)));
 		Category->Items.Add(CreateID(), MakeShareable(new FPlaceableItem(*UActorFactoryPointLight::StaticClass(), SortOrder += 10)));
 		Category->Items.Add(CreateID(), MakeShareable(new FPlaceableItem(*UActorFactoryPlayerStart::StaticClass(), SortOrder += 10)));
-		// Cube
-		Category->Items.Add(CreateID(), MakeShareable(new FPlaceableItem(*UActorFactoryBasicShape::StaticClass(), FAssetData(LoadObject<UStaticMesh>(nullptr, *UActorFactoryBasicShape::BasicCube.ToString())), FName("ClassThumbnail.Cube"), BasicShapeColorOverride, SortOrder += 10, NSLOCTEXT("PlacementMode", "Cube", "Cube"))));
-		// Sphere
-		Category->Items.Add(CreateID(), MakeShareable(new FPlaceableItem(*UActorFactoryBasicShape::StaticClass(), FAssetData(LoadObject<UStaticMesh>(nullptr, *UActorFactoryBasicShape::BasicSphere.ToString())), FName("ClassThumbnail.Sphere"), BasicShapeColorOverride, SortOrder += 10, NSLOCTEXT("PlacementMode", "Sphere", "Sphere"))));
-		// Cylinder
-		Category->Items.Add(CreateID(), MakeShareable(new FPlaceableItem(*UActorFactoryBasicShape::StaticClass(), FAssetData(LoadObject<UStaticMesh>(nullptr, *UActorFactoryBasicShape::BasicCylinder.ToString())), FName("ClassThumbnail.Cylinder"), BasicShapeColorOverride, SortOrder += 10, NSLOCTEXT("PlacementMode", "Cylinder", "Cylinder"))));
-		// Cone
-		Category->Items.Add(CreateID(), MakeShareable(new FPlaceableItem(*UActorFactoryBasicShape::StaticClass(), FAssetData(LoadObject<UStaticMesh>(nullptr, *UActorFactoryBasicShape::BasicCone.ToString())), FName("ClassThumbnail.Cone"), BasicShapeColorOverride, SortOrder += 10, NSLOCTEXT("PlacementMode", "Cone", "Cone"))));
-		// Plane
-		Category->Items.Add(CreateID(), MakeShareable(new FPlaceableItem(*UActorFactoryBasicShape::StaticClass(), FAssetData(LoadObject<UStaticMesh>(nullptr, *UActorFactoryBasicShape::BasicPlane.ToString())), FName("ClassThumbnail.Plane"), BasicShapeColorOverride, SortOrder += 10, NSLOCTEXT("PlacementMode", "Plane", "Plane"))));
-
+	
 		Category->Items.Add(CreateID(), MakeShareable(new FPlaceableItem(*UActorFactoryTriggerBox::StaticClass(), SortOrder += 10)));
 		Category->Items.Add(CreateID(), MakeShareable(new FPlaceableItem(*UActorFactoryTriggerSphere::StaticClass(), SortOrder += 10)));
 	}
@@ -157,6 +149,33 @@ void FPlacementModeModule::StartupModule()
 		Category->Items.Add(CreateID(), MakeShareable(new FPlaceableItem(*UActorFactorySkyLight::StaticClass(), SortOrder += 10)));
 	}
 
+	{
+
+		int32 SortOrder = 0;
+		FName CategoryName = FBuiltInPlacementCategories::Shapes();
+		RegisterPlacementCategory(
+			FPlacementCategoryInfo(
+				NSLOCTEXT("PlacementMode", "Shapes", "Shapes"),
+				FSlateIcon(FAppStyle::GetAppStyleSetName(), "PlacementBrowser.Icons.Basic"),
+				CategoryName,
+				TEXT("PMShapes"),
+				25	
+				)
+			);
+
+		FPlacementCategory* Category = Categories.Find(CategoryName);
+		// Cube
+		Category->Items.Add(CreateID(), MakeShareable(new FPlaceableItem(*UActorFactoryBasicShape::StaticClass(), FAssetData(LoadObject<UStaticMesh>(nullptr, *UActorFactoryBasicShape::BasicCube.ToString())), FName("ClassThumbnail.Cube"), BasicShapeColorOverride, SortOrder += 10, NSLOCTEXT("PlacementMode", "Cube", "Cube"))));
+		// Sphere
+		Category->Items.Add(CreateID(), MakeShareable(new FPlaceableItem(*UActorFactoryBasicShape::StaticClass(), FAssetData(LoadObject<UStaticMesh>(nullptr, *UActorFactoryBasicShape::BasicSphere.ToString())), FName("ClassThumbnail.Sphere"), BasicShapeColorOverride, SortOrder += 10, NSLOCTEXT("PlacementMode", "Sphere", "Sphere"))));
+		// Cylinder
+		Category->Items.Add(CreateID(), MakeShareable(new FPlaceableItem(*UActorFactoryBasicShape::StaticClass(), FAssetData(LoadObject<UStaticMesh>(nullptr, *UActorFactoryBasicShape::BasicCylinder.ToString())), FName("ClassThumbnail.Cylinder"), BasicShapeColorOverride, SortOrder += 10, NSLOCTEXT("PlacementMode", "Cylinder", "Cylinder"))));
+		// Cone
+		Category->Items.Add(CreateID(), MakeShareable(new FPlaceableItem(*UActorFactoryBasicShape::StaticClass(), FAssetData(LoadObject<UStaticMesh>(nullptr, *UActorFactoryBasicShape::BasicCone.ToString())), FName("ClassThumbnail.Cone"), BasicShapeColorOverride, SortOrder += 10, NSLOCTEXT("PlacementMode", "Cone", "Cone"))));
+		// Plane
+		Category->Items.Add(CreateID(), MakeShareable(new FPlaceableItem(*UActorFactoryBasicShape::StaticClass(), FAssetData(LoadObject<UStaticMesh>(nullptr, *UActorFactoryBasicShape::BasicPlane.ToString())), FName("ClassThumbnail.Plane"), BasicShapeColorOverride, SortOrder += 10, NSLOCTEXT("PlacementMode", "Plane", "Plane"))));
+	}
+
 	{
 		int32 SortOrder = 0;
 		FName CategoryName = FBuiltInPlacementCategories::Visual();
@@ -210,6 +229,111 @@ void FPlacementModeModule::StartupModule()
 	BasicShapeThumbnails.Add(UActorFactoryBasicShape::BasicCylinder.ToString(), TEXT("ClassThumbnail.Cylinder"));
 	BasicShapeThumbnails.Add(UActorFactoryBasicShape::BasicCone.ToString(), TEXT("ClassThumbnail.Cone"));
 	BasicShapeThumbnails.Add(UActorFactoryBasicShape::BasicPlane.ToString(), TEXT("ClassThumbnail.Plane"));
+
+
+	// Given a Category Name, this will add a section and all the placeable items in that section directly into the menu
+	auto GenerateQuickCreateSection = [this](const FName& SectionName, int MaxItems = 20)
+	{
+		FPlacementCategory* Category = Categories.Find(SectionName);
+
+		UToolMenu* ContentMenu = UToolMenus::Get()->ExtendMenu("LevelEditor.LevelEditorToolBar.CreateQuickMenu");
+		FToolMenuSection& Section = ContentMenu->AddSection(SectionName, Category->DisplayName);
+
+		int count = 1;
+		for (auto& Pair : Category->Items)
+		{
+			TSharedPtr<const FPlaceableItem> Item = Pair.Value;
+
+			FName TheName = Item->AssetData.AssetName;
+			TSharedRef<SWidget> TheWidget = SNew(SPlacementAssetMenuEntry, Item);
+
+			FToolMenuEntry& NewEntry = Section.AddEntry(
+				FToolMenuEntry::InitWidget(Item->AssetData.AssetName, SNew(SPlacementAssetMenuEntry, Item), FText(), true, true)
+			);
+
+			// NewEntry.StyleNameOverride = "PlaceableItemEntry";
+
+			if (++count > MaxItems)
+				break;
+		}
+	};
+
+	auto GenerateCategorySubMenu = [this](UToolMenu* InMenu, const FName& InSectionName, const FText& InSectionDisplayName, const FName& PlacementCategory)
+	{
+		FToolMenuSection* InSection = InMenu->FindSection(InSectionName);	
+		if (InSection == nullptr)
+		{
+			InSection = &(InMenu->AddSection(InSectionName, InSectionDisplayName));
+		}
+
+		FPlacementCategory* Category = Categories.Find(PlacementCategory);
+
+		FToolMenuEntry& AllSubMenu = InSection->AddSubMenu(PlacementCategory,
+			Category->DisplayName,
+			FText::GetEmpty(),
+			FNewToolMenuDelegate::CreateLambda([this, PlacementCategory](UToolMenu* InMenu)
+			{
+				FToolMenuSection& Section = InMenu->AddSection(PlacementCategory);
+				RegenerateItemsForCategory(PlacementCategory);
+				FPlacementCategory* Category = Categories.Find(PlacementCategory);
+				for (auto& Pair : Category->Items)
+				{
+					TSharedPtr<const FPlaceableItem> Item = Pair.Value;
+
+					FName TheName = Item->AssetData.AssetName;
+					TSharedRef<SWidget> TheWidget = SNew(SPlacementAssetMenuEntry, Item);
+
+					FToolMenuEntry& NewEntry = Section.AddEntry(
+						FToolMenuEntry::InitWidget(Item->AssetData.AssetName, SNew(SPlacementAssetMenuEntry, Item), FText(), true, true)
+					);
+					// NewEntry.StyleNameOverride = "PlaceableItemEntry";
+				}
+			})
+		);
+		AllSubMenu.Icon = Category->DisplayIcon;
+	};
+
+	
+	UToolMenu* ContentMenu = UToolMenus::Get()->ExtendMenu("LevelEditor.LevelEditorToolBar.CreateQuickMenu");
+	FName CreateSectionName = TEXT("PMQCreateMenu");
+	FText CreateSectionDisplayName = NSLOCTEXT("PlacementMode", "PMQCreateMenu", "Create");
+
+	// All Subcategories as submenus
+	FName CategoriesSectionName = TEXT("CreateAllCategories");
+	ContentMenu->AddDynamicSection(CategoriesSectionName,
+		FNewToolMenuDelegate::CreateLambda([this, CreateSectionName, CreateSectionDisplayName, GenerateCategorySubMenu](UToolMenu* InMenu)
+		{
+			TArray<FPlacementCategoryInfo> SortedCategories;
+			GetSortedCategories(SortedCategories);
+			for (auto CategoryInfo : SortedCategories)
+			{
+				// Skip Basic and Recent since we add those later
+				if (CategoryInfo.UniqueHandle == FBuiltInPlacementCategories::Basic() ||
+					CategoryInfo.UniqueHandle == FBuiltInPlacementCategories::RecentlyPlaced())
+					continue;
+
+				GenerateCategorySubMenu(InMenu, CreateSectionName, CreateSectionDisplayName, CategoryInfo.UniqueHandle);
+			}
+		}
+	));
+
+
+	// Basics Section
+	GenerateQuickCreateSection(FBuiltInPlacementCategories::Basic());
+
+	// Recents Section, limit to 5 items, then put in an all Recents submenu
+	RefreshRecentlyPlaced();
+	GenerateQuickCreateSection(FBuiltInPlacementCategories::RecentlyPlaced(), 5);
+	FPlacementCategory* RecentCategory = Categories.Find(FBuiltInPlacementCategories::RecentlyPlaced());
+	const FName& RecentName = FBuiltInPlacementCategories::RecentlyPlaced();
+	GenerateCategorySubMenu(ContentMenu, RecentName, RecentCategory->DisplayName, RecentName);
+
+	// Open Placement Browser Panel
+	FToolMenuSection& BrowserSection = ContentMenu->AddSection(TEXT("PlacementBrowserMenuSection"));
+	BrowserSection.AddSeparator(NAME_None);
+
+	BrowserSection.AddMenuEntry(FLevelEditorCommands::Get().OpenPlaceActors);
+
 }
 
 void FPlacementModeModule::PreUnloadCallback()
diff --git a/Engine/Source/Editor/PlacementMode/Private/SPlacementModeTools.cpp b/Engine/Source/Editor/PlacementMode/Private/SPlacementModeTools.cpp
--- a/Engine/Source/Editor/PlacementMode/Private/SPlacementModeTools.cpp
+++ b/Engine/Source/Editor/PlacementMode/Private/SPlacementModeTools.cpp
@@ -16,6 +16,7 @@
 #include "ContentBrowserDataDragDropOp.h"
 #include "EditorClassUtils.h"
 #include "Widgets/Input/SSearchBox.h"
+#include "ClassIconFinder.h"
 #include "Widgets/Docking/SDockTab.h"
 
 
@@ -306,6 +307,208 @@ const FSlateBrush* SPlacementAssetEntry::GetBorder() const
 	}
 }
 
+
+void SPlacementAssetMenuEntry::Construct(const FArguments& InArgs, const TSharedPtr<const FPlaceableItem>& InItem)
+{	
+	bIsPressed = false;
+
+	check(InItem.IsValid());
+	
+	Item = InItem;
+
+	AssetImage = nullptr;
+
+	TSharedPtr< SHorizontalBox > ActorType = SNew( SHorizontalBox );
+
+	const bool bIsClass = Item->AssetData.GetClass() == UClass::StaticClass();
+	const bool bIsActor = bIsClass ? CastChecked<UClass>(Item->AssetData.GetAsset())->IsChildOf(AActor::StaticClass()) : false;
+
+	AActor* DefaultActor = nullptr;
+	if (Item->Factory != nullptr)
+	{
+		DefaultActor = Item->Factory->GetDefaultActor(Item->AssetData);
+	}
+	else if (bIsActor)
+	{
+		DefaultActor = CastChecked<AActor>(CastChecked<UClass>(Item->AssetData.GetAsset())->ClassDefaultObject);
+	}
+
+	UClass* DocClass = nullptr;
+	TSharedPtr<IToolTip> AssetEntryToolTip;
+	if(DefaultActor != nullptr)
+	{
+		DocClass = DefaultActor->GetClass();
+		AssetEntryToolTip = FEditorClassUtils::GetTooltip(DefaultActor->GetClass());
+	}
+
+	if (!AssetEntryToolTip.IsValid())
+	{
+		AssetEntryToolTip = FSlateApplicationBase::Get().MakeToolTip(Item->DisplayName);
+	}
+	
+	const FButtonStyle& ButtonStyle = FAppStyle::Get().GetWidgetStyle<FButtonStyle>( "Menu.Button" );
+	const float MenuIconSize = FAppStyle::Get().GetFloat("Menu.MenuIconSize");
+
+	Style = &ButtonStyle;
+
+	// Create doc link widget if there is a class to link to
+	TSharedRef<SWidget> DocWidget = SNew(SSpacer);
+	if(DocClass != NULL)
+	{
+		DocWidget = FEditorClassUtils::GetDocumentationLinkWidget(DocClass);
+		DocWidget->SetCursor( EMouseCursor::Default );
+	}
+
+	ChildSlot
+	.HAlign(HAlign_Fill)
+	.VAlign(VAlign_Fill)
+	[
+		SNew(SBorder)
+		.BorderImage( this, &SPlacementAssetMenuEntry::GetBorder )
+		.Cursor( EMouseCursor::GrabHand )
+		.ToolTip( AssetEntryToolTip )
+		.Padding(FMargin(12.f, 3.f, 5.f, 3.f))
+		[
+			SNew( SHorizontalBox )
+
+			+ SHorizontalBox::Slot()
+			.Padding(14.0f, 0.f, 10.f, 0.0f)
+			.AutoWidth()
+			.HAlign(HAlign_Center)
+			.VAlign(VAlign_Center)
+			[
+				SNew(SBox)
+				.WidthOverride(MenuIconSize)
+				.HeightOverride(MenuIconSize)
+				[
+					SNew(SImage)
+					.Image(this, &SPlacementAssetMenuEntry::GetIcon)
+					.ColorAndOpacity(FSlateColor::UseSubduedForeground())
+				]
+			]
+
+			+ SHorizontalBox::Slot()
+			.FillWidth(1.f)
+			.VAlign(VAlign_Center)
+			.HAlign(HAlign_Left)
+			[
+
+				SNew( STextBlock )
+				.ColorAndOpacity(FSlateColor::UseForeground())
+				.Text( Item->DisplayName )
+			]
+
+
+			+ SHorizontalBox::Slot()
+			.VAlign(VAlign_Center)
+			.HAlign(HAlign_Right)
+			.AutoWidth()
+			[
+				SNew(SImage)
+                .ColorAndOpacity(FSlateColor::UseSubduedForeground())
+				.Image(FAppStyle::Get().GetBrush("Icons.DragHandle"))
+			]
+		]
+	];
+}
+
+
+const FSlateBrush* SPlacementAssetMenuEntry::GetIcon() const
+{
+
+	if (AssetImage != nullptr)
+	{
+		return AssetImage;
+	}
+
+	AssetImage = FSlateIconFinder::FindIconBrushForClass(FClassIconFinder::GetIconClassForAssetData(Item->AssetData));
+
+	return AssetImage;
+}
+
+
+FReply SPlacementAssetMenuEntry::OnMouseButtonDown(const FGeometry& MyGeometry, const FPointerEvent& MouseEvent)
+{
+	if ( MouseEvent.GetEffectingButton() == EKeys::LeftMouseButton )
+	{
+		bIsPressed = true;
+
+		return FReply::Handled().DetectDrag( SharedThis( this ), MouseEvent.GetEffectingButton() );
+	}
+
+	return FReply::Unhandled();
+}
+
+FReply SPlacementAssetMenuEntry::OnMouseButtonUp(const FGeometry& MyGeometry, const FPointerEvent& MouseEvent)
+{
+	if ( MouseEvent.GetEffectingButton() == EKeys::LeftMouseButton )
+	{
+		bIsPressed = false;
+	}
+
+	return FReply::Unhandled();
+}
+
+FReply SPlacementAssetMenuEntry::OnDragDetected(const FGeometry& MyGeometry, const FPointerEvent& MouseEvent)
+{
+	bIsPressed = false;
+
+	if (FEditorDelegates::OnAssetDragStarted.IsBound())
+	{
+		TArray<FAssetData> DraggedAssetDatas;
+		DraggedAssetDatas.Add( Item->AssetData );
+		FEditorDelegates::OnAssetDragStarted.Broadcast( DraggedAssetDatas, Item->Factory );
+		return FReply::Handled();
+	}
+
+	if( MouseEvent.IsMouseButtonDown( EKeys::LeftMouseButton ) )
+	{
+		return FReply::Handled().BeginDragDrop( FContentBrowserDataDragDropOp::Legacy_New( MakeArrayView(&Item->AssetData, 1), TArrayView<const FString>(), Item->Factory ) );
+	}
+	else
+	{
+		return FReply::Handled();
+	}
+}
+
+bool SPlacementAssetMenuEntry::IsPressed() const
+{
+	return bIsPressed;
+}
+
+const FSlateBrush* SPlacementAssetMenuEntry::GetBorder() const
+{
+	if ( IsPressed() )
+	{
+		return &(Style->Pressed);
+	}
+	else if ( IsHovered() )
+	{
+		return &(Style->Hovered);
+	}
+	else
+	{
+		return &(Style->Normal);
+	}
+}
+
+FSlateColor SPlacementAssetMenuEntry::GetForegroundColor() const
+{
+	if (IsPressed())
+	{
+		return Style->PressedForeground;
+	}
+	else if (IsHovered())
+	{
+		return Style->HoveredForeground;
+	}
+	else
+	{
+		return Style->NormalForeground;
+	}
+}
+
+
 SPlacementModeTools::~SPlacementModeTools()
 {
 	if ( IPlacementModeModule::IsAvailable() )
diff --git a/Engine/Source/Editor/PlacementMode/Private/SPlacementModeTools.h b/Engine/Source/Editor/PlacementMode/Private/SPlacementModeTools.h
--- a/Engine/Source/Editor/PlacementMode/Private/SPlacementModeTools.h
+++ b/Engine/Source/Editor/PlacementMode/Private/SPlacementModeTools.h
@@ -51,6 +51,36 @@ private:
 	const FSlateBrush* PressedImage;
 };
 
+class SPlacementAssetMenuEntry : public SCompoundWidget
+{
+public:
+	SLATE_BEGIN_ARGS(SPlacementAssetMenuEntry){}
+	SLATE_END_ARGS()
+
+	void Construct(const FArguments& InArgs, const TSharedPtr<const FPlaceableItem>& InItem);
+
+	virtual FReply OnMouseButtonDown(const FGeometry& MyGeometry, const FPointerEvent& MouseEvent) override;
+	virtual FReply OnMouseButtonUp(const FGeometry& MyGeometry, const FPointerEvent& MouseEvent) override;
+	virtual FReply OnDragDetected(const FGeometry& MyGeometry, const FPointerEvent& MouseEvent) override;
+
+	bool IsPressed() const;
+
+	TSharedPtr<const FPlaceableItem> Item;
+
+	virtual FSlateColor GetForegroundColor() const override;
+
+private:
+	const FSlateBrush* GetBorder() const;
+	const FSlateBrush* GetIcon() const;
+
+	bool bIsPressed;
+
+	const FButtonStyle* Style;
+	
+	mutable const FSlateBrush* AssetImage;
+};
+
+
 class SPlacementModeTools : public SCompoundWidget
 {
 public:
diff --git a/Engine/Source/Editor/PlacementMode/Public/IPlacementModeModule.h b/Engine/Source/Editor/PlacementMode/Public/IPlacementModeModule.h
--- a/Engine/Source/Editor/PlacementMode/Public/IPlacementModeModule.h
+++ b/Engine/Source/Editor/PlacementMode/Public/IPlacementModeModule.h
@@ -193,6 +193,7 @@ struct FBuiltInPlacementCategories
 	static FName RecentlyPlaced()	{ static FName Name("RecentlyPlaced");	return Name; }
 	static FName Basic()			{ static FName Name("Basic");			return Name; }
 	static FName Lights()			{ static FName Name("Lights");			return Name; }
+	static FName Shapes()			{ static FName Name("Shapes");			return Name; }
 	static FName Visual()			{ static FName Name("Visual");			return Name; }
 	static FName Volumes()			{ static FName Name("Volumes");			return Name; }
 	static FName AllClasses()		{ static FName Name("AllClasses");		return Name; }
diff --git a/Engine/Source/Runtime/Slate/Private/Framework/MultiBox/SWidgetBlock.cpp b/Engine/Source/Runtime/Slate/Private/Framework/MultiBox/SWidgetBlock.cpp
--- a/Engine/Source/Runtime/Slate/Private/Framework/MultiBox/SWidgetBlock.cpp
+++ b/Engine/Source/Runtime/Slate/Private/Framework/MultiBox/SWidgetBlock.cpp
@@ -91,6 +91,11 @@ void SWidgetBlock::BuildMultiBlockWidget(const ISlateStyle* StyleSet, const FNam
 	{
 		// Account for checkmark used in other menu blocks but not used in for widget rows
 		Padding = Padding + FMargin(14, 0, 8, 0);
+
+		// If there is no label, allow the custom menu widget to consume the entire space
+		if (!bHasLabel)
+			Padding = FMargin(0);
+
 	}
 
 	// Add this widget to the search list of the multibox
@@ -114,7 +119,7 @@ void SWidgetBlock::BuildMultiBlockWidget(const ISlateStyle* StyleSet, const FNam
 			.Visibility( bHasLabel ? EVisibility::Visible : EVisibility::Collapsed )
 			+ SHorizontalBox::Slot()
 			.AutoWidth()
-			.Padding(0.f, 0.f, 4.f, 0.f)
+			.Padding( bHasLabel ? FMargin(0.f, 0.f, 4.f, 0.f) : FMargin(0))
 			.VAlign( VAlign_Center )
 			[
 				SNew( STextBlock )
diff --git a/Engine/Source/Runtime/SlateCore/Private/Styling/StarshipCoreStyle.cpp b/Engine/Source/Runtime/SlateCore/Private/Styling/StarshipCoreStyle.cpp
--- a/Engine/Source/Runtime/SlateCore/Private/Styling/StarshipCoreStyle.cpp
+++ b/Engine/Source/Runtime/SlateCore/Private/Styling/StarshipCoreStyle.cpp
@@ -350,6 +350,7 @@ TSharedRef<ISlateStyle> FStarshipCoreStyle::Create()
 		Style->Set("Icons.Visible", new IMAGE_BRUSH_SVG("Starship/Common/visible", Icon16x16));
 		Style->Set("Icons.Hidden", new IMAGE_BRUSH_SVG("Starship/Common/hidden", Icon16x16));
 
+		Style->Set("Icons.DragHandle", new IMAGE_BRUSH_SVG("Starship/Common/drag-handle", Icon16x16));
 
 	}
 
