Fix issue with PlacementMode startup that was causing cooking errors

diff --git a/Engine/Source/Editor/PlacementMode/Private/PlacementModeModule.cpp b/Engine/Source/Editor/PlacementMode/Private/PlacementModeModule.cpp
--- a/Engine/Source/Editor/PlacementMode/Private/PlacementModeModule.cpp
+++ b/Engine/Source/Editor/PlacementMode/Private/PlacementModeModule.cpp
@@ -231,108 +231,112 @@ void FPlacementModeModule::StartupModule()
 	BasicShapeThumbnails.Add(UActorFactoryBasicShape::BasicPlane.ToString(), TEXT("ClassThumbnail.Plane"));
 
 
-	// Given a Category Name, this will add a section and all the placeable items in that section directly into the menu
-	auto GenerateQuickCreateSection = [this](const FName& SectionName, int MaxItems = 20)
+	if (FSlateApplication::IsInitialized())
 	{
-		FPlacementCategory* Category = Categories.Find(SectionName);
+		// Given a Category Name, this will add a section and all the placeable items in that section directly into the menu
+		auto GenerateQuickCreateSection = [this](const FName& SectionName, int MaxItems = 20)
+		{
+			FPlacementCategory* Category = Categories.Find(SectionName);
 
-		UToolMenu* ContentMenu = UToolMenus::Get()->ExtendMenu("LevelEditor.LevelEditorToolBar.CreateQuickMenu");
-		FToolMenuSection& Section = ContentMenu->AddSection(SectionName, Category->DisplayName);
+			UToolMenu* ContentMenu = UToolMenus::Get()->ExtendMenu("LevelEditor.LevelEditorToolBar.CreateQuickMenu");
+			FToolMenuSection& Section = ContentMenu->AddSection(SectionName, Category->DisplayName);
 
-		int count = 1;
-		for (auto& Pair : Category->Items)
-		{
-			TSharedPtr<const FPlaceableItem> Item = Pair.Value;
+			int count = 1;
+			for (auto& Pair : Category->Items)
+			{
+				TSharedPtr<const FPlaceableItem> Item = Pair.Value;
 
-			FName TheName = Item->AssetData.AssetName;
-			TSharedRef<SWidget> TheWidget = SNew(SPlacementAssetMenuEntry, Item);
+				FName TheName = Item->AssetData.AssetName;
+				TSharedRef<SWidget> TheWidget = SNew(SPlacementAssetMenuEntry, Item);
 
-			FToolMenuEntry& NewEntry = Section.AddEntry(
-				FToolMenuEntry::InitWidget(Item->AssetData.AssetName, SNew(SPlacementAssetMenuEntry, Item), FText(), true, true)
-			);
+				FToolMenuEntry& NewEntry = Section.AddEntry(
+					FToolMenuEntry::InitWidget(Item->AssetData.AssetName, SNew(SPlacementAssetMenuEntry, Item), FText(), true, true)
+				);
 
-			// NewEntry.StyleNameOverride = "PlaceableItemEntry";
+				// NewEntry.StyleNameOverride = "PlaceableItemEntry";
 
-			if (++count > MaxItems)
-				break;
-		}
-	};
+				if (++count > MaxItems)
+					break;
+			}
+		};
 
-	auto GenerateCategorySubMenu = [this](UToolMenu* InMenu, const FName& InSectionName, const FText& InSectionDisplayName, const FName& PlacementCategory)
-	{
-		FToolMenuSection* InSection = InMenu->FindSection(InSectionName);	
-		if (InSection == nullptr)
+		auto GenerateCategorySubMenu = [this](UToolMenu* InMenu, const FName& InSectionName, const FText& InSectionDisplayName, const FName& PlacementCategory)
 		{
-			InSection = &(InMenu->AddSection(InSectionName, InSectionDisplayName));
-		}
+			FToolMenuSection* InSection = InMenu->FindSection(InSectionName);	
+			if (InSection == nullptr)
+			{
+				InSection = &(InMenu->AddSection(InSectionName, InSectionDisplayName));
+			}
 
-		FPlacementCategory* Category = Categories.Find(PlacementCategory);
+			FPlacementCategory* Category = Categories.Find(PlacementCategory);
 
-		FToolMenuEntry& AllSubMenu = InSection->AddSubMenu(PlacementCategory,
-			Category->DisplayName,
-			FText::GetEmpty(),
-			FNewToolMenuDelegate::CreateLambda([this, PlacementCategory](UToolMenu* InMenu)
-			{
-				FToolMenuSection& Section = InMenu->AddSection(PlacementCategory);
-				RegenerateItemsForCategory(PlacementCategory);
-				FPlacementCategory* Category = Categories.Find(PlacementCategory);
-				for (auto& Pair : Category->Items)
+			FToolMenuEntry& AllSubMenu = InSection->AddSubMenu(PlacementCategory,
+				Category->DisplayName,
+				FText::GetEmpty(),
+				FNewToolMenuDelegate::CreateLambda([this, PlacementCategory](UToolMenu* InMenu)
 				{
-					TSharedPtr<const FPlaceableItem> Item = Pair.Value;
-
-					FName TheName = Item->AssetData.AssetName;
-					TSharedRef<SWidget> TheWidget = SNew(SPlacementAssetMenuEntry, Item);
+					FToolMenuSection& Section = InMenu->AddSection(PlacementCategory);
+					RegenerateItemsForCategory(PlacementCategory);
+					FPlacementCategory* Category = Categories.Find(PlacementCategory);
+					for (auto& Pair : Category->Items)
+					{
+						TSharedPtr<const FPlaceableItem> Item = Pair.Value;
+
+						FName TheName = Item->AssetData.AssetName;
+						TSharedRef<SWidget> TheWidget = SNew(SPlacementAssetMenuEntry, Item);
+
+						FToolMenuEntry& NewEntry = Section.AddEntry(
+							FToolMenuEntry::InitWidget(Item->AssetData.AssetName, SNew(SPlacementAssetMenuEntry, Item), FText(), true, true)
+						);
+						// NewEntry.StyleNameOverride = "PlaceableItemEntry";
+					}
+				})
+			);
+			AllSubMenu.Icon = Category->DisplayIcon;
+		};
 
-					FToolMenuEntry& NewEntry = Section.AddEntry(
-						FToolMenuEntry::InitWidget(Item->AssetData.AssetName, SNew(SPlacementAssetMenuEntry, Item), FText(), true, true)
-					);
-					// NewEntry.StyleNameOverride = "PlaceableItemEntry";
-				}
-			})
-		);
-		AllSubMenu.Icon = Category->DisplayIcon;
-	};
+		
+		UToolMenu* ContentMenu = UToolMenus::Get()->ExtendMenu("LevelEditor.LevelEditorToolBar.CreateQuickMenu");
+		FName CreateSectionName = TEXT("PMQCreateMenu");
+		FText CreateSectionDisplayName = NSLOCTEXT("PlacementMode", "PMQCreateMenu", "Create");
 
-	
-	UToolMenu* ContentMenu = UToolMenus::Get()->ExtendMenu("LevelEditor.LevelEditorToolBar.CreateQuickMenu");
-	FName CreateSectionName = TEXT("PMQCreateMenu");
-	FText CreateSectionDisplayName = NSLOCTEXT("PlacementMode", "PMQCreateMenu", "Create");
-
-	// All Subcategories as submenus
-	FName CategoriesSectionName = TEXT("CreateAllCategories");
-	ContentMenu->AddDynamicSection(CategoriesSectionName,
-		FNewToolMenuDelegate::CreateLambda([this, CreateSectionName, CreateSectionDisplayName, GenerateCategorySubMenu](UToolMenu* InMenu)
-		{
-			TArray<FPlacementCategoryInfo> SortedCategories;
-			GetSortedCategories(SortedCategories);
-			for (auto CategoryInfo : SortedCategories)
+		// All Subcategories as submenus
+		FName CategoriesSectionName = TEXT("CreateAllCategories");
+		ContentMenu->AddDynamicSection(CategoriesSectionName,
+			FNewToolMenuDelegate::CreateLambda([this, CreateSectionName, CreateSectionDisplayName, GenerateCategorySubMenu](UToolMenu* InMenu)
 			{
-				// Skip Basic and Recent since we add those later
-				if (CategoryInfo.UniqueHandle == FBuiltInPlacementCategories::Basic() ||
-					CategoryInfo.UniqueHandle == FBuiltInPlacementCategories::RecentlyPlaced())
-					continue;
+				TArray<FPlacementCategoryInfo> SortedCategories;
+				GetSortedCategories(SortedCategories);
+				for (auto CategoryInfo : SortedCategories)
+				{
+					// Skip Basic and Recent since we add those later
+					if (CategoryInfo.UniqueHandle == FBuiltInPlacementCategories::Basic() ||
+						CategoryInfo.UniqueHandle == FBuiltInPlacementCategories::RecentlyPlaced())
+						continue;
 
-				GenerateCategorySubMenu(InMenu, CreateSectionName, CreateSectionDisplayName, CategoryInfo.UniqueHandle);
+					GenerateCategorySubMenu(InMenu, CreateSectionName, CreateSectionDisplayName, CategoryInfo.UniqueHandle);
+				}
 			}
-		}
-	));
+		));
+
 
+		// Basics Section
+		GenerateQuickCreateSection(FBuiltInPlacementCategories::Basic());
 
-	// Basics Section
-	GenerateQuickCreateSection(FBuiltInPlacementCategories::Basic());
+		// Recents Section, limit to 5 items, then put in an all Recents submenu
+		RefreshRecentlyPlaced();
+		GenerateQuickCreateSection(FBuiltInPlacementCategories::RecentlyPlaced(), 5);
+		FPlacementCategory* RecentCategory = Categories.Find(FBuiltInPlacementCategories::RecentlyPlaced());
+		const FName& RecentName = FBuiltInPlacementCategories::RecentlyPlaced();
+		GenerateCategorySubMenu(ContentMenu, RecentName, RecentCategory->DisplayName, RecentName);
 
-	// Recents Section, limit to 5 items, then put in an all Recents submenu
-	RefreshRecentlyPlaced();
-	GenerateQuickCreateSection(FBuiltInPlacementCategories::RecentlyPlaced(), 5);
-	FPlacementCategory* RecentCategory = Categories.Find(FBuiltInPlacementCategories::RecentlyPlaced());
-	const FName& RecentName = FBuiltInPlacementCategories::RecentlyPlaced();
-	GenerateCategorySubMenu(ContentMenu, RecentName, RecentCategory->DisplayName, RecentName);
+		// Open Placement Browser Panel
+		FToolMenuSection& BrowserSection = ContentMenu->AddSection(TEXT("PlacementBrowserMenuSection"));
+		BrowserSection.AddSeparator(NAME_None);
 
-	// Open Placement Browser Panel
-	FToolMenuSection& BrowserSection = ContentMenu->AddSection(TEXT("PlacementBrowserMenuSection"));
-	BrowserSection.AddSeparator(NAME_None);
+		BrowserSection.AddMenuEntry(FLevelEditorCommands::Get().OpenPlaceActors);
 
-	BrowserSection.AddMenuEntry(FLevelEditorCommands::Get().OpenPlaceActors);
+	} // end FSlateApplication::IsInitialized()
 
 }
 
