Quick Place Actors Menu creates an actor at the origin when the menu item is clicked.  Holding CTRL leaves the menu open.

diff --git a/Engine/Source/Editor/PlacementMode/Private/SPlacementModeTools.cpp b/Engine/Source/Editor/PlacementMode/Private/SPlacementModeTools.cpp
--- a/Engine/Source/Editor/PlacementMode/Private/SPlacementModeTools.cpp
+++ b/Engine/Source/Editor/PlacementMode/Private/SPlacementModeTools.cpp
@@ -13,6 +13,7 @@
 #include "EditorModes.h"
 #include "AssetThumbnail.h"
 #include "LevelEditor.h"
+#include "LevelEditorActions.h"
 #include "ContentBrowserDataDragDropOp.h"
 #include "EditorClassUtils.h"
 #include "Widgets/Input/SSearchBox.h"
@@ -444,6 +445,15 @@ FReply SPlacementAssetMenuEntry::OnMouseButtonUp(const FGeometry& MyGeometry, co
 	if ( MouseEvent.GetEffectingButton() == EKeys::LeftMouseButton )
 	{
 		bIsPressed = false;
+
+		FTransform TempTransform;
+		FLevelEditorActionCallbacks::AddActor(Item->Factory, Item->AssetData, &TempTransform);
+		if (!MouseEvent.IsControlDown())
+		{
+			FSlateApplication::Get().DismissAllMenus();
+		}
+
+		return FReply::Handled();
 	}
 
 	return FReply::Unhandled();
