Placement Mode: shape category special icon handling and updates to plugins using FPlaceableItem

diff --git a/Engine/Plugins/Experimental/ColorCorrectRegions/Source/ColorCorrectRegionsEditor/Private/ColorCorrectRegionsEditorModule.cpp b/Engine/Plugins/Experimental/ColorCorrectRegions/Source/ColorCorrectRegionsEditor/Private/ColorCorrectRegionsEditorModule.cpp
--- a/Engine/Plugins/Experimental/ColorCorrectRegions/Source/ColorCorrectRegionsEditor/Private/ColorCorrectRegionsEditorModule.cpp
+++ b/Engine/Plugins/Experimental/ColorCorrectRegions/Source/ColorCorrectRegionsEditor/Private/ColorCorrectRegionsEditorModule.cpp
@@ -27,6 +27,7 @@ void FColorCorrectRegionsEditorModule::OnPlacementModeRefresh(FName CategoryName
 		FPlaceableItem* CCRPlaceableItem = new FPlaceableItem(
 			*UActorFactoryBlueprint::StaticClass(),
 			FAssetData(CCRBlueprint, true),
+			FName("CCR.PlaceActorThumbnail"),
 			FName("CCR.PlaceActorIcon"),
 			TOptional<FLinearColor>(),
 			TOptional<int32>(),
diff --git a/Engine/Plugins/Experimental/ColorCorrectRegions/Source/ColorCorrectRegionsEditor/Private/ColorCorrectRegionsStyle.cpp b/Engine/Plugins/Experimental/ColorCorrectRegions/Source/ColorCorrectRegionsEditor/Private/ColorCorrectRegionsStyle.cpp
--- a/Engine/Plugins/Experimental/ColorCorrectRegions/Source/ColorCorrectRegionsEditor/Private/ColorCorrectRegionsStyle.cpp
+++ b/Engine/Plugins/Experimental/ColorCorrectRegions/Source/ColorCorrectRegionsEditor/Private/ColorCorrectRegionsStyle.cpp
@@ -18,7 +18,8 @@ void FColorCorrectRegionsStyle::Initialize()
 	CCRStyle = MakeShared<FSlateStyleSet>(FName("ColorCorrectRegionsStyle"));
 
 	FString IconPath = IPluginManager::Get().FindPlugin(TEXT("ColorCorrectRegions"))->GetBaseDir() + TEXT("/Resources/PlaceActorPreview.png");
-	CCRStyle->Set("CCR.PlaceActorIcon", new FSlateImageBrush(IconPath, FVector2D(40.0f, 40.0f)));
+	CCRStyle->Set("CCR.PlaceActorThumbnail", new FSlateImageBrush(IconPath, FVector2D(40.0f, 40.0f)));
+	CCRStyle->Set("CCR.PlaceActorIcon", new FSlateImageBrush(IconPath, FVector2D(16.0f, 16.0f)));
 
 	FSlateStyleRegistry::RegisterSlateStyle(*CCRStyle.Get());
 }
diff --git a/Engine/Plugins/Experimental/VirtualCamera/Source/VirtualCameraEditor/Private/VirtualCameraEditorModule.cpp b/Engine/Plugins/Experimental/VirtualCamera/Source/VirtualCameraEditor/Private/VirtualCameraEditorModule.cpp
--- a/Engine/Plugins/Experimental/VirtualCamera/Source/VirtualCameraEditor/Private/VirtualCameraEditorModule.cpp
+++ b/Engine/Plugins/Experimental/VirtualCamera/Source/VirtualCameraEditor/Private/VirtualCameraEditorModule.cpp
@@ -83,6 +83,7 @@ public:
 				IPlacementModeModule::Get().RegisterPlaceableItem(Info->UniqueHandle, MakeShared<FPlaceableItem>(
 					*UActorFactoryBlueprint::StaticClass(),
 					SimpleVirtualCameraAssetData,
+					FName("ClassThumbnail.CameraActor"),
 					FName("ClassIcon.CameraActor"),
 					TOptional<FLinearColor>(),
 					TOptional<int32>(0),
@@ -100,6 +101,7 @@ public:
 				IPlacementModeModule::Get().RegisterPlaceableItem(Info->UniqueHandle, MakeShared<FPlaceableItem>(
 					*UActorFactoryBlueprint::StaticClass(),
 					VirtualCamera2ActorAssetData,
+					FName("ClassThumbnail.CameraActor"),
 					FName("ClassIcon.CameraActor"),
 					TOptional<FLinearColor>(),
 					TOptional<int32>(1),
diff --git a/Engine/Plugins/Runtime/HDRIBackdrop/Source/HDRIBackdrop/Private/HDRIBackdropPlacement.cpp b/Engine/Plugins/Runtime/HDRIBackdrop/Source/HDRIBackdrop/Private/HDRIBackdropPlacement.cpp
--- a/Engine/Plugins/Runtime/HDRIBackdrop/Source/HDRIBackdrop/Private/HDRIBackdropPlacement.cpp
+++ b/Engine/Plugins/Runtime/HDRIBackdrop/Source/HDRIBackdrop/Private/HDRIBackdropPlacement.cpp
@@ -20,6 +20,7 @@ void FHDRIBackdropPlacement::RegisterPlacement()
 	FPlaceableItem* BPPlacement = new FPlaceableItem(
 		*UActorFactoryBlueprint::StaticClass(),
 		FAssetData(HDRIBackdrop, true),
+		FName("HDRIBackdrop.ModesThumbnail"),
 		FName("HDRIBackdrop.ModesIcon"),
 		TOptional<FLinearColor>(),
 		TOptional<int32>(),
diff --git a/Engine/Plugins/Runtime/HDRIBackdrop/Source/HDRIBackdrop/Private/HDRIBackdropStyle.cpp b/Engine/Plugins/Runtime/HDRIBackdrop/Source/HDRIBackdrop/Private/HDRIBackdropStyle.cpp
--- a/Engine/Plugins/Runtime/HDRIBackdrop/Source/HDRIBackdrop/Private/HDRIBackdropStyle.cpp
+++ b/Engine/Plugins/Runtime/HDRIBackdrop/Source/HDRIBackdrop/Private/HDRIBackdropStyle.cpp
@@ -21,7 +21,8 @@ void FHDRIBackdropStyle::Initialize()
 	StyleSet->SetCoreContentRoot(FPaths::EngineContentDir() / TEXT("Slate"));
 
 	FString IconPath = IPluginManager::Get().FindPlugin(TEXT("HDRIBackdrop"))->GetBaseDir() + TEXT("/Resources/HDRIBackdrop.png");
-	StyleSet->Set("HDRIBackdrop.ModesIcon", new FSlateImageBrush(IconPath, FVector2D(40.0f, 40.0f)));
+	StyleSet->Set("HDRIBackdrop.ModesThumbnail", new FSlateImageBrush(IconPath, FVector2D(40.0f, 40.0f)));
+	StyleSet->Set("HDRIBackdrop.ModesIcon", new FSlateImageBrush(IconPath, FVector2D(16.0f, 16.0f)));
 
 	FSlateStyleRegistry::RegisterSlateStyle(*StyleSet.Get());
 }
diff --git a/Engine/Plugins/Runtime/SunPosition/Source/SunPosition/Private/SunPositionPlacement.cpp b/Engine/Plugins/Runtime/SunPosition/Source/SunPosition/Private/SunPositionPlacement.cpp
--- a/Engine/Plugins/Runtime/SunPosition/Source/SunPosition/Private/SunPositionPlacement.cpp
+++ b/Engine/Plugins/Runtime/SunPosition/Source/SunPosition/Private/SunPositionPlacement.cpp
@@ -23,6 +23,7 @@ void FSunPositionPlacement::RegisterPlacement()
 		FPlaceableItem* BPPlacement = new FPlaceableItem(
 			*UActorFactoryBlueprint::StaticClass(),
 			FAssetData(SunSky, true),
+			FName("SunPosition.ModesThumbnail"),
 			FName("SunPosition.ModesIcon"),
 			TOptional<FLinearColor>(),
 			TOptional<int32>(),
diff --git a/Engine/Plugins/Runtime/SunPosition/Source/SunPosition/Private/SunPositionStyle.cpp b/Engine/Plugins/Runtime/SunPosition/Source/SunPosition/Private/SunPositionStyle.cpp
--- a/Engine/Plugins/Runtime/SunPosition/Source/SunPosition/Private/SunPositionStyle.cpp
+++ b/Engine/Plugins/Runtime/SunPosition/Source/SunPosition/Private/SunPositionStyle.cpp
@@ -21,7 +21,8 @@ void FSunPositionStyle::Initialize()
 	StyleSet->SetCoreContentRoot(FPaths::EngineContentDir() / TEXT("Slate"));
 
 	FString IconPath = IPluginManager::Get().FindPlugin(TEXT("SunPosition"))->GetBaseDir() + TEXT("/Resources/SunPosition.png");
-	StyleSet->Set("SunPosition.ModesIcon", new FSlateImageBrush(IconPath, FVector2D(40.0f, 40.0f)));
+	StyleSet->Set("SunPosition.ModesThumbnail", new FSlateImageBrush(IconPath, FVector2D(40.0f, 40.0f)));
+	StyleSet->Set("SunPosition.ModesIcon", new FSlateImageBrush(IconPath, FVector2D(16.0f, 16.0f)));
 
 	FSlateStyleRegistry::RegisterSlateStyle(*StyleSet.Get());
 }
diff --git a/Engine/Plugins/VirtualProduction/CompositePlane/Source/CompositePlane/Private/CompositePlanePlacement.cpp b/Engine/Plugins/VirtualProduction/CompositePlane/Source/CompositePlane/Private/CompositePlanePlacement.cpp
--- a/Engine/Plugins/VirtualProduction/CompositePlane/Source/CompositePlane/Private/CompositePlanePlacement.cpp
+++ b/Engine/Plugins/VirtualProduction/CompositePlane/Source/CompositePlane/Private/CompositePlanePlacement.cpp
@@ -25,6 +25,7 @@ void FCompositePlanePlacement::RegisterPlacement()
 		*UActorFactoryBlueprint::StaticClass(),
 		FAssetData(CompositePlane, true),
 		FName(""),
+		FName(""),
 		TOptional<FLinearColor>(),
 		TOptional<int32>(),
 		NSLOCTEXT("PlacementMode", "Composite Plane", "Composite Plane")
diff --git a/Engine/Plugins/VirtualProduction/DataCharts/Source/DataChartsEditor/Private/DataChartsPlacement.cpp b/Engine/Plugins/VirtualProduction/DataCharts/Source/DataChartsEditor/Private/DataChartsPlacement.cpp
--- a/Engine/Plugins/VirtualProduction/DataCharts/Source/DataChartsEditor/Private/DataChartsPlacement.cpp
+++ b/Engine/Plugins/VirtualProduction/DataCharts/Source/DataChartsEditor/Private/DataChartsPlacement.cpp
@@ -37,6 +37,7 @@ void FDataChartsPlacement::RegisterPlacement()
 			*UActorFactoryBlueprint::StaticClass(),
 			FAssetData(ChartBP, true),
 			It.Key(),
+			It.Key(),
 			TOptional<FLinearColor>(),
 			TOptional<int32>(),
 			FText::FromString(ChartName)
diff --git a/Engine/Source/Editor/PlacementMode/Private/PlacementModeModule.cpp b/Engine/Source/Editor/PlacementMode/Private/PlacementModeModule.cpp
--- a/Engine/Source/Editor/PlacementMode/Private/PlacementModeModule.cpp
+++ b/Engine/Source/Editor/PlacementMode/Private/PlacementModeModule.cpp
@@ -156,7 +156,7 @@ void FPlacementModeModule::StartupModule()
 		RegisterPlacementCategory(
 			FPlacementCategoryInfo(
 				NSLOCTEXT("PlacementMode", "Shapes", "Shapes"),
-				FSlateIcon(FAppStyle::GetAppStyleSetName(), "PlacementBrowser.Icons.Basic"),
+				FSlateIcon(FAppStyle::GetAppStyleSetName(), "ClassIcon.Cube"),
 				CategoryName,
 				TEXT("PMShapes"),
 				25	
@@ -165,15 +165,15 @@ void FPlacementModeModule::StartupModule()
 
 		FPlacementCategory* Category = Categories.Find(CategoryName);
 		// Cube
-		Category->Items.Add(CreateID(), MakeShareable(new FPlaceableItem(*UActorFactoryBasicShape::StaticClass(), FAssetData(LoadObject<UStaticMesh>(nullptr, *UActorFactoryBasicShape::BasicCube.ToString())), FName("ClassThumbnail.Cube"), BasicShapeColorOverride, SortOrder += 10, NSLOCTEXT("PlacementMode", "Cube", "Cube"))));
+		Category->Items.Add(CreateID(), MakeShareable(new FPlaceableItem(*UActorFactoryBasicShape::StaticClass(), FAssetData(LoadObject<UStaticMesh>(nullptr, *UActorFactoryBasicShape::BasicCube.ToString())), FName("ClassThumbnail.Cube"), FName("ClassIcon.Cube"), BasicShapeColorOverride, SortOrder += 10, NSLOCTEXT("PlacementMode", "Cube", "Cube"))));
 		// Sphere
-		Category->Items.Add(CreateID(), MakeShareable(new FPlaceableItem(*UActorFactoryBasicShape::StaticClass(), FAssetData(LoadObject<UStaticMesh>(nullptr, *UActorFactoryBasicShape::BasicSphere.ToString())), FName("ClassThumbnail.Sphere"), BasicShapeColorOverride, SortOrder += 10, NSLOCTEXT("PlacementMode", "Sphere", "Sphere"))));
+		Category->Items.Add(CreateID(), MakeShareable(new FPlaceableItem(*UActorFactoryBasicShape::StaticClass(), FAssetData(LoadObject<UStaticMesh>(nullptr, *UActorFactoryBasicShape::BasicSphere.ToString())), FName("ClassThumbnail.Sphere"), FName("ClassIcon.Sphere"), BasicShapeColorOverride, SortOrder += 10, NSLOCTEXT("PlacementMode", "Sphere", "Sphere"))));
 		// Cylinder
-		Category->Items.Add(CreateID(), MakeShareable(new FPlaceableItem(*UActorFactoryBasicShape::StaticClass(), FAssetData(LoadObject<UStaticMesh>(nullptr, *UActorFactoryBasicShape::BasicCylinder.ToString())), FName("ClassThumbnail.Cylinder"), BasicShapeColorOverride, SortOrder += 10, NSLOCTEXT("PlacementMode", "Cylinder", "Cylinder"))));
+		Category->Items.Add(CreateID(), MakeShareable(new FPlaceableItem(*UActorFactoryBasicShape::StaticClass(), FAssetData(LoadObject<UStaticMesh>(nullptr, *UActorFactoryBasicShape::BasicCylinder.ToString())), FName("ClassThumbnail.Cylinder"), FName("ClassIcon.Cylinder"), BasicShapeColorOverride, SortOrder += 10, NSLOCTEXT("PlacementMode", "Cylinder", "Cylinder"))));
 		// Cone
-		Category->Items.Add(CreateID(), MakeShareable(new FPlaceableItem(*UActorFactoryBasicShape::StaticClass(), FAssetData(LoadObject<UStaticMesh>(nullptr, *UActorFactoryBasicShape::BasicCone.ToString())), FName("ClassThumbnail.Cone"), BasicShapeColorOverride, SortOrder += 10, NSLOCTEXT("PlacementMode", "Cone", "Cone"))));
+		Category->Items.Add(CreateID(), MakeShareable(new FPlaceableItem(*UActorFactoryBasicShape::StaticClass(), FAssetData(LoadObject<UStaticMesh>(nullptr, *UActorFactoryBasicShape::BasicCone.ToString())), FName("ClassThumbnail.Cone"), FName("ClassIcon.Cone"), BasicShapeColorOverride, SortOrder += 10, NSLOCTEXT("PlacementMode", "Cone", "Cone"))));
 		// Plane
-		Category->Items.Add(CreateID(), MakeShareable(new FPlaceableItem(*UActorFactoryBasicShape::StaticClass(), FAssetData(LoadObject<UStaticMesh>(nullptr, *UActorFactoryBasicShape::BasicPlane.ToString())), FName("ClassThumbnail.Plane"), BasicShapeColorOverride, SortOrder += 10, NSLOCTEXT("PlacementMode", "Plane", "Plane"))));
+		Category->Items.Add(CreateID(), MakeShareable(new FPlaceableItem(*UActorFactoryBasicShape::StaticClass(), FAssetData(LoadObject<UStaticMesh>(nullptr, *UActorFactoryBasicShape::BasicPlane.ToString())), FName("ClassThumbnail.Plane"), FName("ClassIcon.Plane"), BasicShapeColorOverride, SortOrder += 10, NSLOCTEXT("PlacementMode", "Plane", "Plane"))));
 	}
 
 	{
@@ -638,11 +638,11 @@ void FPlacementModeModule::RefreshAllPlaceableClasses()
 	Category->Items.Add(CreateID(), MakeShareable(new FPlaceableItem(*UActorFactoryEmptyActor::StaticClass())));
 	Category->Items.Add(CreateID(), MakeShareable(new FPlaceableItem(*UActorFactoryCharacter::StaticClass())));
 	Category->Items.Add(CreateID(), MakeShareable(new FPlaceableItem(*UActorFactoryPawn::StaticClass())));
-	Category->Items.Add(CreateID(), MakeShareable(new FPlaceableItem(*UActorFactoryBasicShape::StaticClass(), FAssetData(LoadObject<UStaticMesh>(nullptr, *UActorFactoryBasicShape::BasicCube.ToString())), FName("ClassThumbnail.Cube"), GetBasicShapeColorOverride(), TOptional<int32>(), NSLOCTEXT("PlacementMode", "Cube", "Cube"))));
-	Category->Items.Add(CreateID(), MakeShareable(new FPlaceableItem(*UActorFactoryBasicShape::StaticClass(), FAssetData(LoadObject<UStaticMesh>(nullptr, *UActorFactoryBasicShape::BasicSphere.ToString())), FName("ClassThumbnail.Sphere"), GetBasicShapeColorOverride(), TOptional<int32>(), NSLOCTEXT("PlacementMode", "Sphere", "Sphere"))));
-	Category->Items.Add(CreateID(), MakeShareable(new FPlaceableItem(*UActorFactoryBasicShape::StaticClass(), FAssetData(LoadObject<UStaticMesh>(nullptr, *UActorFactoryBasicShape::BasicCylinder.ToString())), FName("ClassThumbnail.Cylinder"), GetBasicShapeColorOverride(), TOptional<int32>(), NSLOCTEXT("PlacementMode", "Cylinder", "Cylinder"))));
-	Category->Items.Add(CreateID(), MakeShareable(new FPlaceableItem(*UActorFactoryBasicShape::StaticClass(), FAssetData(LoadObject<UStaticMesh>(nullptr, *UActorFactoryBasicShape::BasicCone.ToString())), FName("ClassThumbnail.Cone"), GetBasicShapeColorOverride(), TOptional<int32>(), NSLOCTEXT("PlacementMode", "Cone", "Cone"))));
-	Category->Items.Add(CreateID(), MakeShareable(new FPlaceableItem(*UActorFactoryBasicShape::StaticClass(), FAssetData(LoadObject<UStaticMesh>(nullptr, *UActorFactoryBasicShape::BasicPlane.ToString())), FName("ClassThumbnail.Plane"), GetBasicShapeColorOverride(), TOptional<int32>(), NSLOCTEXT("PlacementMode", "Plane", "Plane"))));
+	Category->Items.Add(CreateID(), MakeShareable(new FPlaceableItem(*UActorFactoryBasicShape::StaticClass(), FAssetData(LoadObject<UStaticMesh>(nullptr, *UActorFactoryBasicShape::BasicCube.ToString())), FName("ClassThumbnail.Cube"), FName("ClassIcon.Cube"), GetBasicShapeColorOverride(), TOptional<int32>(), NSLOCTEXT("PlacementMode", "Cube", "Cube"))));
+	Category->Items.Add(CreateID(), MakeShareable(new FPlaceableItem(*UActorFactoryBasicShape::StaticClass(), FAssetData(LoadObject<UStaticMesh>(nullptr, *UActorFactoryBasicShape::BasicSphere.ToString())), FName("ClassThumbnail.Sphere"), FName("ClassIcon.Sphere"), GetBasicShapeColorOverride(), TOptional<int32>(), NSLOCTEXT("PlacementMode", "Sphere", "Sphere"))));
+	Category->Items.Add(CreateID(), MakeShareable(new FPlaceableItem(*UActorFactoryBasicShape::StaticClass(), FAssetData(LoadObject<UStaticMesh>(nullptr, *UActorFactoryBasicShape::BasicCylinder.ToString())), FName("ClassThumbnail.Cylinder"), FName("ClassIcon.Cylinder"), GetBasicShapeColorOverride(), TOptional<int32>(), NSLOCTEXT("PlacementMode", "Cylinder", "Cylinder"))));
+	Category->Items.Add(CreateID(), MakeShareable(new FPlaceableItem(*UActorFactoryBasicShape::StaticClass(), FAssetData(LoadObject<UStaticMesh>(nullptr, *UActorFactoryBasicShape::BasicCone.ToString())), FName("ClassThumbnail.Cone"), FName("ClassIcon.Cone"), GetBasicShapeColorOverride(), TOptional<int32>(), NSLOCTEXT("PlacementMode", "Cone", "Cone"))));
+	Category->Items.Add(CreateID(), MakeShareable(new FPlaceableItem(*UActorFactoryBasicShape::StaticClass(), FAssetData(LoadObject<UStaticMesh>(nullptr, *UActorFactoryBasicShape::BasicPlane.ToString())), FName("ClassThumbnail.Plane"), FName("ClassIcon.Plane"), GetBasicShapeColorOverride(), TOptional<int32>(), NSLOCTEXT("PlacementMode", "Plane", "Plane"))));
 
 	// Make a map of UClasses to ActorFactories that support them
 	const TArray< UActorFactory *>& ActorFactories = GEditor->ActorFactories;
diff --git a/Engine/Source/Editor/PlacementMode/Private/SPlacementModeTools.cpp b/Engine/Source/Editor/PlacementMode/Private/SPlacementModeTools.cpp
--- a/Engine/Source/Editor/PlacementMode/Private/SPlacementModeTools.cpp
+++ b/Engine/Source/Editor/PlacementMode/Private/SPlacementModeTools.cpp
@@ -314,7 +314,7 @@ void SPlacementAssetMenuEntry::Construct(const FArguments& InArgs, const TShared
 	bIsPressed = false;
 
 	check(InItem.IsValid());
-	
+
 	Item = InItem;
 
 	AssetImage = nullptr;
@@ -422,7 +422,14 @@ const FSlateBrush* SPlacementAssetMenuEntry::GetIcon() const
 		return AssetImage;
 	}
 
-	AssetImage = FSlateIconFinder::FindIconBrushForClass(FClassIconFinder::GetIconClassForAssetData(Item->AssetData));
+	if (Item->ClassIconBrushOverride != NAME_None)
+	{
+		AssetImage = FSlateIconFinder::FindCustomIconBrushForClass(nullptr, TEXT("ClassIcon"), Item->ClassIconBrushOverride);
+	}
+	else
+	{
+		AssetImage = FSlateIconFinder::FindIconBrushForClass(FClassIconFinder::GetIconClassForAssetData(Item->AssetData));
+	}
 
 	return AssetImage;
 }
diff --git a/Engine/Source/Editor/PlacementMode/Public/IPlacementModeModule.h b/Engine/Source/Editor/PlacementMode/Public/IPlacementModeModule.h
--- a/Engine/Source/Editor/PlacementMode/Public/IPlacementModeModule.h
+++ b/Engine/Source/Editor/PlacementMode/Public/IPlacementModeModule.h
@@ -109,6 +109,7 @@ struct FPlaceableItem
 		UClass& InAssetClass,
 		const FAssetData& InAssetData,
 		FName InClassThumbnailBrushOverride = NAME_None,
+		FName InClassIconBrushOverride = NAME_None,
 		TOptional<FLinearColor> InAssetTypeColorOverride = TOptional<FLinearColor>(),
 		TOptional<int32> InSortOrder = TOptional<int32>(),
 		TOptional<FText> InDisplayName = TOptional<FText>()
@@ -116,6 +117,7 @@ struct FPlaceableItem
 		: Factory(GEditor->FindActorFactoryByClass(&InAssetClass))
 		, AssetData(InAssetData)
 		, ClassThumbnailBrushOverride(InClassThumbnailBrushOverride)
+		, ClassIconBrushOverride(InClassIconBrushOverride)
 		, AssetTypeColorOverride(InAssetTypeColorOverride)
 		, SortOrder(InSortOrder)
 	{
@@ -177,6 +179,9 @@ struct FPlaceableItem
 	/** Optional override for the thumbnail brush (passed to FClassIconFinder::FindThumbnailForClass in the form ClassThumbnail.<override>) */
 	FName ClassThumbnailBrushOverride;
 
+	/** Optional override for the small icon brush */
+	FName ClassIconBrushOverride;
+
 	/** Whether to always use the generic thumbnail for this item or not */
 	bool bAlwaysUseGenericThumbnail;
 
