Fixes dragable AssetCard from Quick Place Actors Menu

diff --git a/Engine/Source/Editor/PlacementMode/Private/SPlacementModeTools.cpp b/Engine/Source/Editor/PlacementMode/Private/SPlacementModeTools.cpp
--- a/Engine/Source/Editor/PlacementMode/Private/SPlacementModeTools.cpp
+++ b/Engine/Source/Editor/PlacementMode/Private/SPlacementModeTools.cpp
@@ -480,7 +480,7 @@ FReply SPlacementAssetMenuEntry::OnDragDetected(const FGeometry& MyGeometry, con
 
 	if( MouseEvent.IsMouseButtonDown( EKeys::LeftMouseButton ) )
 	{
-		return FReply::Handled().BeginDragDrop( FContentBrowserDataDragDropOp::Legacy_New( MakeArrayView(&Item->AssetData, 1), TArrayView<const FString>(), Item->Factory ) );
+		return FReply::Handled().BeginDragDrop(FAssetDragDropOp::New(Item->AssetData, Item->Factory));
 	}
 	else
 	{
