Create Menu / Place Actors Menu fixes

diff --git a/Engine/Source/Editor/PlacementMode/Private/PlacementModeModule.cpp b/Engine/Source/Editor/PlacementMode/Private/PlacementModeModule.cpp
--- a/Engine/Source/Editor/PlacementMode/Private/PlacementModeModule.cpp
+++ b/Engine/Source/Editor/PlacementMode/Private/PlacementModeModule.cpp
@@ -234,12 +234,10 @@ void FPlacementModeModule::StartupModule()
 	if (FSlateApplication::IsInitialized())
 	{
 		// Given a Category Name, this will add a section and all the placeable items in that section directly into the menu
-		auto GenerateQuickCreateSection = [this](const FName& SectionName, int MaxItems = 20)
+		auto GenerateQuickCreateSection = [this](const FName& SectionName, UToolMenu* InMenu, int MaxItems = 20)
 		{
 			FPlacementCategory* Category = Categories.Find(SectionName);
-
-			UToolMenu* ContentMenu = UToolMenus::Get()->ExtendMenu("LevelEditor.LevelEditorToolBar.CreateQuickMenu");
-			FToolMenuSection& Section = ContentMenu->AddSection(SectionName, Category->DisplayName);
+			FToolMenuSection& Section = InMenu->AddSection(SectionName, Category->DisplayName);
 
 			int count = 1;
 			for (auto& Pair : Category->Items)
@@ -253,8 +251,6 @@ void FPlacementModeModule::StartupModule()
 					FToolMenuEntry::InitWidget(Item->AssetData.AssetName, SNew(SPlacementAssetMenuEntry, Item), FText(), true, true)
 				);
 
-				// NewEntry.StyleNameOverride = "PlaceableItemEntry";
-
 				if (++count > MaxItems)
 					break;
 			}
@@ -262,37 +258,40 @@ void FPlacementModeModule::StartupModule()
 
 		auto GenerateCategorySubMenu = [this](UToolMenu* InMenu, const FName& InSectionName, const FText& InSectionDisplayName, const FName& PlacementCategory)
 		{
-			FToolMenuSection* InSection = InMenu->FindSection(InSectionName);	
-			if (InSection == nullptr)
-			{
-				InSection = &(InMenu->AddSection(InSectionName, InSectionDisplayName));
-			}
-
+			
+			RegenerateItemsForCategory(PlacementCategory);
 			FPlacementCategory* Category = Categories.Find(PlacementCategory);
-
-			FToolMenuEntry& AllSubMenu = InSection->AddSubMenu(PlacementCategory,
-				Category->DisplayName,
-				FText::GetEmpty(),
-				FNewToolMenuDelegate::CreateLambda([this, PlacementCategory](UToolMenu* InMenu)
+			if (!Category->Items.IsEmpty())
+			{
+				FToolMenuSection* InSection = InMenu->FindSection(InSectionName);	
+				if (InSection == nullptr)
 				{
-					FToolMenuSection& Section = InMenu->AddSection(PlacementCategory);
-					RegenerateItemsForCategory(PlacementCategory);
-					FPlacementCategory* Category = Categories.Find(PlacementCategory);
-					for (auto& Pair : Category->Items)
-					{
-						TSharedPtr<const FPlaceableItem> Item = Pair.Value;
-
-						FName TheName = Item->AssetData.AssetName;
-						TSharedRef<SWidget> TheWidget = SNew(SPlacementAssetMenuEntry, Item);
+					InSection = &(InMenu->AddSection(InSectionName, InSectionDisplayName));
+				}
 
-						FToolMenuEntry& NewEntry = Section.AddEntry(
-							FToolMenuEntry::InitWidget(Item->AssetData.AssetName, SNew(SPlacementAssetMenuEntry, Item), FText(), true, true)
-						);
-						// NewEntry.StyleNameOverride = "PlaceableItemEntry";
-					}
-				})
-			);
-			AllSubMenu.Icon = Category->DisplayIcon;
+				FToolMenuEntry& AllSubMenu = InSection->AddSubMenu(PlacementCategory,
+					Category->DisplayName,
+					FText::GetEmpty(),
+					FNewToolMenuDelegate::CreateLambda([this, PlacementCategory](UToolMenu* InMenu)
+					{
+						FToolMenuSection& Section = InMenu->AddSection(PlacementCategory);
+						RegenerateItemsForCategory(PlacementCategory);
+						FPlacementCategory* Category = Categories.Find(PlacementCategory);
+						for (auto& Pair : Category->Items)
+						{
+							TSharedPtr<const FPlaceableItem> Item = Pair.Value;
+
+							FName TheName = Item->AssetData.AssetName;
+							TSharedRef<SWidget> TheWidget = SNew(SPlacementAssetMenuEntry, Item);
+
+							FToolMenuEntry& NewEntry = Section.AddEntry(
+								FToolMenuEntry::InitWidget(Item->AssetData.AssetName, SNew(SPlacementAssetMenuEntry, Item), FText(), true, true)
+							);
+						}
+					})
+				);
+				AllSubMenu.Icon = Category->DisplayIcon;
+			}
 		};
 
 		
@@ -321,14 +320,20 @@ void FPlacementModeModule::StartupModule()
 
 
 		// Basics Section
-		GenerateQuickCreateSection(FBuiltInPlacementCategories::Basic());
+		GenerateQuickCreateSection(FBuiltInPlacementCategories::Basic(), ContentMenu);
 
 		// Recents Section, limit to 5 items, then put in an all Recents submenu
-		RefreshRecentlyPlaced();
-		GenerateQuickCreateSection(FBuiltInPlacementCategories::RecentlyPlaced(), 5);
-		FPlacementCategory* RecentCategory = Categories.Find(FBuiltInPlacementCategories::RecentlyPlaced());
 		const FName& RecentName = FBuiltInPlacementCategories::RecentlyPlaced();
-		GenerateCategorySubMenu(ContentMenu, RecentName, RecentCategory->DisplayName, RecentName);
+		ContentMenu->AddDynamicSection(RecentName,
+			FNewToolMenuDelegate::CreateLambda([this, GenerateQuickCreateSection, GenerateCategorySubMenu](UToolMenu* InMenu)
+			{
+				const FName& RecentName = FBuiltInPlacementCategories::RecentlyPlaced();
+				FPlacementCategory* RecentCategory = Categories.Find(RecentName);
+				RefreshRecentlyPlaced();
+				GenerateQuickCreateSection(RecentName, InMenu, 5);
+				GenerateCategorySubMenu(InMenu, RecentName, RecentCategory->DisplayName, RecentName);
+			}
+		));
 
 		// Open Placement Browser Panel
 		FToolMenuSection& BrowserSection = ContentMenu->AddSection(TEXT("PlacementBrowserMenuSection"));
