Create Menu Cleanup - Place actor in front of camera on click

UETOOL-3054

diff --git a/Engine/Source/Editor/PlacementMode/Private/SPlacementModeTools.cpp b/Engine/Source/Editor/PlacementMode/Private/SPlacementModeTools.cpp
--- a/Engine/Source/Editor/PlacementMode/Private/SPlacementModeTools.cpp
+++ b/Engine/Source/Editor/PlacementMode/Private/SPlacementModeTools.cpp
@@ -14,6 +14,7 @@
 #include "AssetThumbnail.h"
 #include "LevelEditor.h"
 #include "LevelEditorActions.h"
+#include "LevelEditorViewport.h"
 #include "ContentBrowserDataDragDropOp.h"
 #include "EditorClassUtils.h"
 #include "Widgets/Input/SSearchBox.h"
@@ -454,7 +455,14 @@ FReply SPlacementAssetMenuEntry::OnMouseButtonUp(const FGeometry& MyGeometry, co
 		bIsPressed = false;
 
 		FTransform TempTransform;
-		FLevelEditorActionCallbacks::AddActor(Item->Factory, Item->AssetData, &TempTransform);
+		if( AActor* NewActor = FLevelEditorActionCallbacks::AddActor(Item->Factory, Item->AssetData, &TempTransform) )
+		{
+  			GEditor->MoveActorInFrontOfCamera(*NewActor, 
+  				GCurrentLevelEditingViewportClient->GetViewLocation(), 
+  				GCurrentLevelEditingViewportClient->GetViewRotation().Vector()
+  			);
+		}
+
 		if (!MouseEvent.IsControlDown())
 		{
 			FSlateApplication::Get().DismissAllMenus();
