Create Menu Add Search for top Level Window, + misc cleanup of redundant menu rows.

diff --git a/Engine/Source/Editor/LevelEditor/Private/LevelEditorToolBar.cpp b/Engine/Source/Editor/LevelEditor/Private/LevelEditorToolBar.cpp
--- a/Engine/Source/Editor/LevelEditor/Private/LevelEditorToolBar.cpp
+++ b/Engine/Source/Editor/LevelEditor/Private/LevelEditorToolBar.cpp
@@ -1889,6 +1889,7 @@ void FLevelEditorToolBar::RegisterCreateMenu()
 {
 #define LOCTEXT_NAMESPACE "LevelToolBarViewMenu"
 	UToolMenu* ContentMenu = UToolMenus::Get()->RegisterMenu("LevelEditor.LevelEditorToolBar.CreateQuickMenu");
+	ContentMenu->bSearchable = true;
 	{
 		FToolMenuSection& Section = ContentMenu->FindOrAddSection("Create");
 		Section.InitSection("ProjectContent", LOCTEXT("Create_Label", "Create"), FToolMenuInsert());
diff --git a/Engine/Source/Editor/PlacementMode/Private/PlacementModeModule.cpp b/Engine/Source/Editor/PlacementMode/Private/PlacementModeModule.cpp
--- a/Engine/Source/Editor/PlacementMode/Private/PlacementModeModule.cpp
+++ b/Engine/Source/Editor/PlacementMode/Private/PlacementModeModule.cpp
@@ -243,10 +243,6 @@ void FPlacementModeModule::StartupModule()
 			for (auto& Pair : Category->Items)
 			{
 				TSharedPtr<const FPlaceableItem> Item = Pair.Value;
-
-				FName TheName = Item->AssetData.AssetName;
-				TSharedRef<SWidget> TheWidget = SNew(SPlacementAssetMenuEntry, Item);
-
 				FToolMenuEntry& NewEntry = Section.AddEntry(
 					FToolMenuEntry::InitWidget(Item->AssetData.AssetName, SNew(SPlacementAssetMenuEntry, Item), FText(), true, true)
 				);
@@ -280,10 +276,6 @@ void FPlacementModeModule::StartupModule()
 						for (auto& Pair : Category->Items)
 						{
 							TSharedPtr<const FPlaceableItem> Item = Pair.Value;
-
-							FName TheName = Item->AssetData.AssetName;
-							TSharedRef<SWidget> TheWidget = SNew(SPlacementAssetMenuEntry, Item);
-
 							FToolMenuEntry& NewEntry = Section.AddEntry(
 								FToolMenuEntry::InitWidget(Item->AssetData.AssetName, SNew(SPlacementAssetMenuEntry, Item), FText(), true, true)
 							);
