Remove obsoiete code from placement tool

diff --git a/Engine/Source/Editor/PlacementMode/Private/SPlacementModeTools.cpp b/Engine/Source/Editor/PlacementMode/Private/SPlacementModeTools.cpp
--- a/Engine/Source/Editor/PlacementMode/Private/SPlacementModeTools.cpp
+++ b/Engine/Source/Editor/PlacementMode/Private/SPlacementModeTools.cpp
@@ -551,10 +551,6 @@ void SPlacementModeTools::Construct( const FArguments& InArgs, TSharedRef<SDockT
 
 	ParentTab->SetOnTabDrawerOpened(FSimpleDelegate::CreateSP(this, &SPlacementModeTools::OnTabDrawerOpened));
 
-
-	static const FName PlacementBrowserActiveTabBarBrushName("PlacementBrowser.ActiveTabBar");
-	PlacementBrowserActiveTabBarBrush = FEditorStyle::GetBrush(PlacementBrowserActiveTabBarBrushName);
-
 	SearchTextFilter = MakeShareable(new FPlacementAssetEntryTextFilter(
 		FPlacementAssetEntryTextFilter::FItemToStringArray::CreateStatic(&PlacementViewFilter::GetBasicStrings)
 		));
@@ -689,41 +685,6 @@ void SPlacementModeTools::Construct( const FArguments& InArgs, TSharedRef<SDockT
 	PlacementModeModule.OnAllPlaceableAssetsChanged().AddSP( this, &SPlacementModeTools::UpdatePlaceableAssets );
 }
 
-TSharedRef< SWidget > SPlacementModeTools::CreatePlacementGroupTab( const FPlacementCategoryInfo& Info )
-{
-	return SNew( SCheckBox )
-	.Style( FEditorStyle::Get(), "PlacementBrowser.Tab" )
-	.OnCheckStateChanged( this, &SPlacementModeTools::OnPlacementTabChanged, Info.UniqueHandle )
-	.IsChecked( this, &SPlacementModeTools::GetPlacementTabCheckedState, Info.UniqueHandle )
-	[
-		SNew( SOverlay )
-
-		+ SOverlay::Slot()
-		.VAlign( VAlign_Center )
-		[
-			SNew(SSpacer)
-			.Size( FVector2D( 1, 30 ) )
-		]
-
-		+ SOverlay::Slot()
-		.Padding( FMargin(6, 0, 15, 0) )
-		.VAlign( VAlign_Center )
-		[
-			SNew( STextBlock )
-			.TextStyle( FEditorStyle::Get(), "PlacementBrowser.Tab.Text" )
-			.Text( Info.DisplayName )
-		]
-
-		+ SOverlay::Slot()
-		.VAlign( VAlign_Fill )
-		.HAlign( HAlign_Left )
-		[
-			SNew(SImage)
-			.Image( this, &SPlacementModeTools::PlacementGroupBorderImage, Info.UniqueHandle )
-		]
-	];
-}
-
 FName SPlacementModeTools::GetActiveTab() const
 {
 	return IsSearchActive() ? FBuiltInPlacementCategories::AllClasses() : ActiveTabName;
@@ -830,29 +791,6 @@ void SPlacementModeTools::OnTabDrawerOpened()
 	FSlateApplication::Get().SetKeyboardFocus(SearchBoxPtr, EFocusCause::SetDirectly);
 }
 
-void SPlacementModeTools::OnPlacementTabChanged( ECheckBoxState NewState, FName CategoryName )
-{
-	if ( NewState == ECheckBoxState::Checked )
-	{
-		ActiveTabName = CategoryName;
-		IPlacementModeModule::Get().RegenerateItemsForCategory(ActiveTabName);
-
-		bNeedsUpdate = true;
-	}
-}
-
-const FSlateBrush* SPlacementModeTools::PlacementGroupBorderImage( FName CategoryName ) const
-{
-	if ( ActiveTabName == CategoryName )
-	{
-		return PlacementBrowserActiveTabBarBrush;
-	}
-	else
-	{
-		return nullptr;
-	}
-}
-
 void SPlacementModeTools::UpdateRecentlyPlacedAssets( const TArray< FActorPlacementInfo >& RecentlyPlaced )
 {
 	if (GetActiveTab() == FBuiltInPlacementCategories::RecentlyPlaced())
diff --git a/Engine/Source/Editor/PlacementMode/Private/SPlacementModeTools.h b/Engine/Source/Editor/PlacementMode/Private/SPlacementModeTools.h
--- a/Engine/Source/Editor/PlacementMode/Private/SPlacementModeTools.h
+++ b/Engine/Source/Editor/PlacementMode/Private/SPlacementModeTools.h
@@ -99,9 +99,6 @@ private:
 
 private:
 
-	/** Creates a tab widget to show on the left that when clicked sets the currently active tab */
-	TSharedRef< SWidget > CreatePlacementGroupTab( const FPlacementCategoryInfo& Info );
-
 	/** Generates a widget for the specified item */
 	TSharedRef<ITableRow> OnGenerateWidgetForItem(TSharedPtr<FPlaceableItem> InItem, const TSharedRef<STableViewBase>& OwnerTable);
 
@@ -120,11 +117,7 @@ private:
 	/** Called when the list of placeable assets changes. */
 	void UpdatePlaceableAssets();
 
-	/** Gets the border image for the tab, this is the 'active' orange bar. */
-	const FSlateBrush* PlacementGroupBorderImage( FName CategoryName ) const;
-
 	/** When the tab is clicked we adjust the check state, so that the right style is displayed. */
-	void OnPlacementTabChanged( ECheckBoxState NewState, FName CategoryName );
 	void OnCategoryChanged(const ECheckBoxState NewState, FName InCategory);
 
 	/** Called when the placement mode tools are opened from a tab drawer */
@@ -186,6 +179,4 @@ private:
 
 	/** List view that shows placeable items */
 	TSharedPtr<SListView<TSharedPtr<FPlaceableItem>>> ListView;
-
-	const FSlateBrush* PlacementBrowserActiveTabBarBrush;
 };
