Manual merge from UE4
Placement Tool:
- Add category blacklist
- Don't trigger OnAllPlaceableAssetsChanged event during initial loading of assets

diff --git a/Engine/Source/Editor/PlacementMode/Private/PlacementModeModule.cpp b/Engine/Source/Editor/PlacementMode/Private/PlacementModeModule.cpp
--- a/Engine/Source/Editor/PlacementMode/Private/PlacementModeModule.cpp
+++ b/Engine/Source/Editor/PlacementMode/Private/PlacementModeModule.cpp
@@ -5,6 +5,7 @@
 #include "Modules/ModuleManager.h"
 #include "UObject/Object.h"
 #include "Misc/Guid.h"
+#include "Misc/BlacklistNames.h"
 #include "UObject/Class.h"
 #include "UObject/UObjectHash.h"
 #include "UObject/UObjectIterator.h"
@@ -72,23 +73,33 @@ TOptional<FLinearColor> GetBasicShapeColorOverride()
 	return BasicShapeColorOverride;
 }
 
+FPlacementModeModule::FPlacementModeModule()
+	: CategoryBlacklist(MakeShareable(new FBlacklistNames()))
+{
+	CategoryBlacklist->OnFilterChanged().AddRaw(this, &FPlacementModeModule::OnCategoryBlacklistChanged);
+}
+
 void FPlacementModeModule::StartupModule()
 {
 	TArray< FString > RecentlyPlacedAsStrings;
 	GConfig->GetArray(TEXT("PlacementMode"), TEXT("RecentlyPlaced"), RecentlyPlacedAsStrings, GEditorPerProjectIni);
 
-	//FString ActivePaletteName;
-	//GConfig->GetString( TEXT( "PlacementMode" ), TEXT( "ActivePalette" ), ActivePaletteName, GEditorPerProjectIni );
-
 	for (int Index = 0; Index < RecentlyPlacedAsStrings.Num(); Index++)
 	{
 		RecentlyPlaced.Add(FActorPlacementInfo(RecentlyPlacedAsStrings[Index]));
 	}
 
-	FAssetRegistryModule& AssetRegistryModule = FModuleManager::LoadModuleChecked<FAssetRegistryModule>(TEXT("AssetRegistry"));
-	AssetRegistryModule.Get().OnAssetRemoved().AddRaw(this, &FPlacementModeModule::OnAssetRemoved);
-	AssetRegistryModule.Get().OnAssetRenamed().AddRaw(this, &FPlacementModeModule::OnAssetRenamed);
-	AssetRegistryModule.Get().OnAssetAdded().AddRaw(this, &FPlacementModeModule::OnAssetAdded);
+	IAssetRegistry& AssetRegistry = FModuleManager::LoadModuleChecked<FAssetRegistryModule>(TEXT("AssetRegistry")).Get();
+	AssetRegistry.OnAssetRemoved().AddRaw(this, &FPlacementModeModule::OnAssetRemoved);
+	AssetRegistry.OnAssetRenamed().AddRaw(this, &FPlacementModeModule::OnAssetRenamed);
+	if (AssetRegistry.IsLoadingAssets())
+	{
+		AssetRegistry.OnFilesLoaded().AddRaw(this, &FPlacementModeModule::OnInitialAssetsScanComplete);
+	}
+	else
+	{
+		AssetRegistry.OnAssetAdded().AddRaw(this, &FPlacementModeModule::OnAssetAdded);
+	}
 
 	TOptional<FLinearColor> BasicShapeColorOverride = GetBasicShapeColorOverride();
 
@@ -345,6 +356,7 @@ void FPlacementModeModule::PreUnloadCallback()
 		AssetRegistryModule->Get().OnAssetRemoved().RemoveAll(this);
 		AssetRegistryModule->Get().OnAssetRenamed().RemoveAll(this);
 		AssetRegistryModule->Get().OnAssetAdded().RemoveAll(this);
+		AssetRegistryModule->Get().OnFilesLoaded().RemoveAll(this);
 	}
 }
 
@@ -440,6 +452,15 @@ void FPlacementModeModule::OnAssetAdded(const FAssetData& AssetData)
 	AllPlaceableAssetsChanged.Broadcast();
 }
 
+void FPlacementModeModule::OnInitialAssetsScanComplete()
+{
+	IAssetRegistry& AssetRegistry = FModuleManager::LoadModuleChecked<FAssetRegistryModule>(TEXT("AssetRegistry")).Get();
+	AssetRegistry.OnAssetAdded().AddRaw(this, &FPlacementModeModule::OnAssetAdded);
+	AssetRegistry.OnFilesLoaded().RemoveAll(this);
+
+	AllPlaceableAssetsChanged.Broadcast();
+}
+
 void FPlacementModeModule::AddToRecentlyPlaced(UObject* Asset, UActorFactory* FactoryUsed /* = NULL */)
 {
 	TArray< UObject* > Assets;
@@ -460,15 +481,19 @@ bool FPlacementModeModule::RegisterPlacementCategory(const FPlacementCategoryInf
 	}
 
 	Categories.Add(Info.UniqueHandle, Info);
+	PlacementModeCategoryListChanged.Broadcast();
 	return true;
 }
 
 void FPlacementModeModule::UnregisterPlacementCategory(FName Handle)
 {
-	Categories.Remove(Handle);
+	if (Categories.Remove(Handle))
+	{
+		PlacementModeCategoryListChanged.Broadcast();
+	}
 }
 
-void FPlacementModeModule::GetSortedCategories(TArray<FPlacementCategoryInfo>& OutCategories) const
+void FPlacementModeModule::GetUserFacingCategories(TArray<FPlacementCategoryInfo>& OutCategories) const
 {
 	TArray<FName> SortedNames;
 	Categories.GenerateKeyArray(SortedNames);
@@ -480,7 +505,10 @@ void FPlacementModeModule::GetSortedCategories(TArray<FPlacementCategoryInfo>& O
 	OutCategories.Reset(Categories.Num());
 	for (const FName& Name : SortedNames)
 	{
-		OutCategories.Add(Categories[Name]);
+		if (CategoryBlacklist->PassesFilter(Name))
+		{
+			OutCategories.Add(Categories[Name]);
+		}
 	}
 }
 
@@ -698,4 +726,9 @@ FPlacementModeID FPlacementModeModule::CreateID(FName InCategory)
 	NewID.UniqueID = CreateID();
 	NewID.Category = InCategory;
 	return NewID;
+}
+
+void FPlacementModeModule::OnCategoryBlacklistChanged()
+{
+	PlacementModeCategoryListChanged.Broadcast();
 }
\ No newline at end of file
diff --git a/Engine/Source/Editor/PlacementMode/Private/SPlacementModeTools.cpp b/Engine/Source/Editor/PlacementMode/Private/SPlacementModeTools.cpp
--- a/Engine/Source/Editor/PlacementMode/Private/SPlacementModeTools.cpp
+++ b/Engine/Source/Editor/PlacementMode/Private/SPlacementModeTools.cpp
@@ -538,8 +538,10 @@ SPlacementModeTools::~SPlacementModeTools()
 {
 	if ( IPlacementModeModule::IsAvailable() )
 	{
-		IPlacementModeModule::Get().OnRecentlyPlacedChanged().RemoveAll( this );
-		IPlacementModeModule::Get().OnAllPlaceableAssetsChanged().RemoveAll( this );
+		IPlacementModeModule& PlacementModeModule = IPlacementModeModule::Get();
+		PlacementModeModule.OnRecentlyPlacedChanged().RemoveAll(this);
+		PlacementModeModule.OnAllPlaceableAssetsChanged().RemoveAll(this);
+		PlacementModeModule.OnPlacementModeCategoryListChanged().RemoveAll(this);
 	}
 }
 
@@ -548,6 +550,7 @@ void SPlacementModeTools::Construct( const FArguments& InArgs, TSharedRef<SDockT
 	bPlaceablesFullRefreshRequested = false;
 	bRecentlyPlacedRefreshRequested = false;
 	bNeedsUpdate = true;
+	ActiveTabName = FBuiltInPlacementCategories::Basic();
 
 	ParentTab->SetOnTabDrawerOpened(FSimpleDelegate::CreateSP(this, &SPlacementModeTools::OnTabDrawerOpened));
 
@@ -559,28 +562,7 @@ void SPlacementModeTools::Construct( const FArguments& InArgs, TSharedRef<SDockT
 	.HAlign(HAlign_Center)
 	.SlotPadding(FMargin(2.0f, 0.0f));
 
-	// Populate the categories and body from the defined placeable items
-	IPlacementModeModule& PlacementModeModule = IPlacementModeModule::Get();
-
-	TArray<FPlacementCategoryInfo> Categories;
-	PlacementModeModule.GetSortedCategories(Categories);
-	for (const FPlacementCategoryInfo& Category : Categories)
-	{
-		CategoryFilterPtr->AddSlot()
-		[
-			SNew(SCheckBox)
-			.Padding(FMargin(4.f, 4.f))
-			.Style( &FAppStyle::Get(),  "PaletteToolBar.Tab" )
-			.OnCheckStateChanged(this, &SPlacementModeTools::OnCategoryChanged, Category.UniqueHandle)
-			.IsChecked(this, &SPlacementModeTools::GetPlacementTabCheckedState, Category.UniqueHandle)
-			.ToolTipText(Category.DisplayName)
-			[
-				  SNew(SImage)
-	              .ColorAndOpacity(FSlateColor::UseForeground())
-	              .Image(Category.DisplayIcon.GetIcon())
-			]
-		];
-	}
+	UpdatePlacementCategories();
 
 	TSharedRef<SScrollBar> ScrollBar = SNew(SScrollBar)
 		.Thickness(FVector2D(9.0f, 9.0f));
@@ -678,11 +660,10 @@ void SPlacementModeTools::Construct( const FArguments& InArgs, TSharedRef<SDockT
 		]
 	];
 
-	ActiveTabName = FBuiltInPlacementCategories::Basic();
-	bNeedsUpdate = true;
-
-	PlacementModeModule.OnRecentlyPlacedChanged().AddSP( this, &SPlacementModeTools::UpdateRecentlyPlacedAssets );
-	PlacementModeModule.OnAllPlaceableAssetsChanged().AddSP( this, &SPlacementModeTools::UpdatePlaceableAssets );
+	IPlacementModeModule& PlacementModeModule = IPlacementModeModule::Get();
+	PlacementModeModule.OnRecentlyPlacedChanged().AddSP(this, &SPlacementModeTools::UpdateRecentlyPlacedAssets);
+	PlacementModeModule.OnAllPlaceableAssetsChanged().AddSP(this, &SPlacementModeTools::UpdatePlaceableAssets);
+	PlacementModeModule.OnPlacementModeCategoryListChanged().AddSP(this, &SPlacementModeTools::UpdatePlacementCategories);
 }
 
 FName SPlacementModeTools::GetActiveTab() const
@@ -690,6 +671,16 @@ FName SPlacementModeTools::GetActiveTab() const
 	return IsSearchActive() ? FBuiltInPlacementCategories::AllClasses() : ActiveTabName;
 }
 
+void SPlacementModeTools::SetActiveTab(FName TabName)
+{
+	if (TabName != ActiveTabName)
+	{
+		ActiveTabName = TabName;
+		IPlacementModeModule::Get().RegenerateItemsForCategory(ActiveTabName);
+		bNeedsUpdate = true;
+	}
+}
+
 void SPlacementModeTools::UpdateFilteredItems()
 {
 	bNeedsUpdate = false;
@@ -777,12 +768,9 @@ TSharedRef<ITableRow> SPlacementModeTools::OnGenerateWidgetForItem(TSharedPtr<FP
 
 void SPlacementModeTools::OnCategoryChanged(const ECheckBoxState NewState, FName InCategory)
 {
-	if (NewState == ECheckBoxState::Checked && InCategory != ActiveTabName)
+	if (NewState == ECheckBoxState::Checked)
 	{
-		ActiveTabName = InCategory;
-		IPlacementModeModule::Get().RegenerateItemsForCategory(ActiveTabName);
-
-		bNeedsUpdate = true;
+		SetActiveTab(InCategory);
 	}
 }
 
@@ -807,6 +795,58 @@ void SPlacementModeTools::UpdatePlaceableAssets()
 	}
 }
 
+void SPlacementModeTools::UpdatePlacementCategories()
+{
+	bool BasicTabExists = false;
+	FName TabToActivate;
+
+	CategoryFilterPtr->ClearChildren();
+
+	TArray<FPlacementCategoryInfo> Categories;
+	IPlacementModeModule::Get().GetUserFacingCategories(Categories);
+	for (const FPlacementCategoryInfo& Category : Categories)
+	{
+		if (Category.UniqueHandle == FBuiltInPlacementCategories::Basic())
+		{
+			BasicTabExists = true;
+		}
+
+		if (Category.UniqueHandle == ActiveTabName)
+		{
+			TabToActivate = ActiveTabName;
+		}
+
+		CategoryFilterPtr->AddSlot()
+		[
+			SNew(SCheckBox)
+			.Padding(FMargin(4.f, 4.f))
+			.Style( &FAppStyle::Get(),  "PaletteToolBar.Tab" )
+			.OnCheckStateChanged(this, &SPlacementModeTools::OnCategoryChanged, Category.UniqueHandle)
+			.IsChecked(this, &SPlacementModeTools::GetPlacementTabCheckedState, Category.UniqueHandle)
+			.ToolTipText(Category.DisplayName)
+			[
+				SNew(SImage)
+				.ColorAndOpacity(FSlateColor::UseForeground())
+				.Image(Category.DisplayIcon.GetIcon())
+			]
+		];
+	}
+
+	if (TabToActivate.IsNone())
+	{
+		if (BasicTabExists)
+		{
+			TabToActivate = FBuiltInPlacementCategories::Basic();
+		}
+		else if (Categories.Num() > 0)
+		{
+			TabToActivate = Categories[0].UniqueHandle;
+		}
+	}
+
+	SetActiveTab(TabToActivate);
+}
+
 void SPlacementModeTools::Tick( const FGeometry& AllottedGeometry, const double InCurrentTime, const float InDeltaTime )
 {
 	if ( bPlaceablesFullRefreshRequested )
diff --git a/Engine/Source/Editor/PlacementMode/Private/SPlacementModeTools.h b/Engine/Source/Editor/PlacementMode/Private/SPlacementModeTools.h
--- a/Engine/Source/Editor/PlacementMode/Private/SPlacementModeTools.h
+++ b/Engine/Source/Editor/PlacementMode/Private/SPlacementModeTools.h
@@ -105,6 +105,9 @@ private:
 	/** Get the identifier of the currently active tab */
 	FName GetActiveTab() const;
 
+	/** Changes the active tab */
+	void SetActiveTab(FName TabName);
+
 	/** Check if a search is active */
 	bool IsSearchActive() const;
 
@@ -117,6 +120,9 @@ private:
 	/** Called when the list of placeable assets changes. */
 	void UpdatePlaceableAssets();
 
+	/** Called when the list of placement categories changes. */
+	void UpdatePlacementCategories();
+
 	/** When the tab is clicked we adjust the check state, so that the right style is displayed. */
 	void OnCategoryChanged(const ECheckBoxState NewState, FName InCategory);
 
diff --git a/Engine/Source/Editor/PlacementMode/Public/IPlacementModeModule.h b/Engine/Source/Editor/PlacementMode/Public/IPlacementModeModule.h
--- a/Engine/Source/Editor/PlacementMode/Public/IPlacementModeModule.h
+++ b/Engine/Source/Editor/PlacementMode/Public/IPlacementModeModule.h
@@ -14,6 +14,8 @@
 #include "GameFramework/Volume.h"
 #include "Editor.h"
 
+class FBlacklistNames;
+
 /**
  * Struct that defines an identifier for a particular placeable item in this module.
  * Only be obtainable through IPlacementModeModule::RegisterPlaceableItem
@@ -244,6 +246,12 @@ public:
 	 */
 	virtual const TArray< FActorPlacementInfo >& GetRecentlyPlaced() const = 0;
 
+	/**
+	 * @return the event that is broadcast whenever the user facing list of placement mode categories gets modified
+	 */
+	DECLARE_EVENT(IPlacementMode, FOnPlacementModeCategoryListChanged);
+	virtual FOnPlacementModeCategoryListChanged& OnPlacementModeCategoryListChanged() = 0;
+
 	/**
 	 * @return the event that is broadcast whenever a placement mode category is refreshed
 	 */
@@ -276,6 +284,13 @@ public:
 	 */
 	virtual bool RegisterPlacementCategory(const FPlacementCategoryInfo& Info) = 0;
 
+	/**
+	 * Unregister a previously registered category
+	 *
+	 * @param UniqueHandle	The unique handle of the category to unregister
+	 */
+	virtual void UnregisterPlacementCategory(FName Handle) = 0;
+
 	/**
 	 * Retrieve an already registered category
 	 *
@@ -284,19 +299,20 @@ public:
 	 */
 	virtual const FPlacementCategoryInfo* GetRegisteredPlacementCategory(FName UniqueHandle) const = 0;
 
+	/** Placement categories blacklist */
+	virtual TSharedRef<FBlacklistNames>& GetCategoryBlacklist() = 0;
+
 	/**
-	 * Populate the specified array with all registered category information, sorted by SortOrder
+	 * Populate the specified array with all registered category information that isn't blacklisted, sorted by SortOrder
 	 *
 	 * @param OutCategories	The array to populate with registered category inforamtion
 	 */
-	virtual void GetSortedCategories(TArray<FPlacementCategoryInfo>& OutCategories) const = 0;
+	virtual void GetUserFacingCategories(TArray<FPlacementCategoryInfo>& OutCategories) const = 0;
 
 	/**
-	 * Unregister a previously registered category
-	 *
-	 * @param UniqueHandle	The unique handle of the category to unregister
+	 * @deprecated Use GetUserFacingCategories instead 
 	 */
-	virtual void UnregisterPlacementCategory(FName Handle) = 0;
+	void GetSortedCategories(TArray<FPlacementCategoryInfo>& OutCategories) const { GetUserFacingCategories(OutCategories); }
 
 	/**
 	 * Register a new placeable item for the specified category
diff --git a/Engine/Source/Editor/PlacementMode/Public/PlacementModeModule.h b/Engine/Source/Editor/PlacementMode/Public/PlacementModeModule.h
--- a/Engine/Source/Editor/PlacementMode/Public/PlacementModeModule.h
+++ b/Engine/Source/Editor/PlacementMode/Public/PlacementModeModule.h
@@ -36,6 +36,8 @@ class FPlacementModeModule : public IPlacementModeModule
 {
 public:
 
+	FPlacementModeModule();
+
 	/**
 	 * Called right after the module's DLL has been loaded and the module object has been created
 	 */
@@ -46,6 +48,9 @@ public:
 	 */
 	virtual void PreUnloadCallback() override;
 
+	DECLARE_DERIVED_EVENT(FPlacementModeModule, IPlacementModeModule::FOnPlacementModeCategoryListChanged, FOnPlacementModeCategoryListChanged);
+	virtual FOnPlacementModeCategoryListChanged& OnPlacementModeCategoryListChanged() override { return PlacementModeCategoryListChanged; }
+
 	DECLARE_DERIVED_EVENT(FPlacementModeModule, IPlacementModeModule::FOnPlacementModeCategoryRefreshed, FOnPlacementModeCategoryRefreshed);
 	virtual FOnPlacementModeCategoryRefreshed& OnPlacementModeCategoryRefreshed() override { return PlacementModeCategoryRefreshed; }
 
@@ -60,12 +65,6 @@ public:
 	 */
 	virtual void AddToRecentlyPlaced(const TArray< UObject* >& PlacedObjects, UActorFactory* FactoryUsed = NULL) override;
 
-	void OnAssetRemoved(const FAssetData& /*InRemovedAssetData*/);
-
-	void OnAssetRenamed(const FAssetData& AssetData, const FString& OldObjectPath);
-
-	void OnAssetAdded(const FAssetData& AssetData);
-
 	/**
 	 * Add the specified asset to the recently placed items list
 	 */
@@ -90,7 +89,9 @@ public:
 
 	virtual void UnregisterPlacementCategory(FName Handle);
 
-	virtual void GetSortedCategories(TArray<FPlacementCategoryInfo>& OutCategories) const;
+	virtual TSharedRef<FBlacklistNames>& GetCategoryBlacklist() override { return CategoryBlacklist; }
+
+	virtual void GetUserFacingCategories(TArray<FPlacementCategoryInfo>& OutCategories) const;
 
 	virtual TOptional<FPlacementModeID> RegisterPlaceableItem(FName CategoryName, const TSharedRef<FPlaceableItem>& InItem);
 
@@ -104,6 +105,14 @@ public:
 
 private:
 
+	void OnAssetRemoved(const FAssetData& /*InRemovedAssetData*/);
+
+	void OnAssetRenamed(const FAssetData& AssetData, const FString& OldObjectPath);
+
+	void OnAssetAdded(const FAssetData& AssetData);
+
+	void OnInitialAssetsScanComplete();
+
 	void RefreshRecentlyPlaced();
 
 	void RefreshVolumes();
@@ -114,15 +123,20 @@ private:
 
 	FPlacementModeID CreateID(FName InCategory);
 
+	void OnCategoryBlacklistChanged();
+
 private:
 
 	TMap<FName, FPlacementCategory> Categories;
 
-	TArray< FActorPlacementInfo >	RecentlyPlaced;
-	FOnRecentlyPlacedChanged		RecentlyPlacedChanged;
+	TSharedRef<FBlacklistNames> CategoryBlacklist;
+
+	TArray< FActorPlacementInfo > RecentlyPlaced;
+	FOnRecentlyPlacedChanged RecentlyPlacedChanged;
 
-	FOnAllPlaceableAssetsChanged	AllPlaceableAssetsChanged;
+	FOnAllPlaceableAssetsChanged AllPlaceableAssetsChanged;
 	FOnPlacementModeCategoryRefreshed PlacementModeCategoryRefreshed;
+	FOnPlacementModeCategoryListChanged PlacementModeCategoryListChanged;
 
 	TArray< TSharedPtr<FExtender> > ContentPaletteFiltersExtenders;
 	TArray< TSharedPtr<FExtender> > PaletteExtenders;
