Manual merge from UE4
Placement tool item predicates

diff --git a/Engine/Source/Editor/PlacementMode/Private/PlacementModeModule.cpp b/Engine/Source/Editor/PlacementMode/Private/PlacementModeModule.cpp
--- a/Engine/Source/Editor/PlacementMode/Private/PlacementModeModule.cpp
+++ b/Engine/Source/Editor/PlacementMode/Private/PlacementModeModule.cpp
@@ -493,7 +493,7 @@ void FPlacementModeModule::UnregisterPlacementCategory(FName Handle)
 	}
 }
 
-void FPlacementModeModule::GetUserFacingCategories(TArray<FPlacementCategoryInfo>& OutCategories) const
+void FPlacementModeModule::GetSortedCategories(TArray<FPlacementCategoryInfo>& OutCategories) const
 {
 	TArray<FName> SortedNames;
 	Categories.GenerateKeyArray(SortedNames);
@@ -533,6 +533,26 @@ void FPlacementModeModule::UnregisterPlaceableItem(FPlacementModeID ID)
 	}
 }
 
+bool FPlacementModeModule::RegisterPlaceableItemFilter(TPlaceableItemPredicate Predicate, FName OwnerName)
+{
+	TPlaceableItemPredicate& ExistingPredicate = PlaceableItemPredicates.FindOrAdd(OwnerName);
+	if (!ExistingPredicate)
+	{
+		ExistingPredicate = MoveTemp(Predicate);
+		PlaceableItemFilteringChanged.Broadcast();
+		return true;
+	}
+	return false;
+}
+
+void FPlacementModeModule::UnregisterPlaceableItemFilter(FName OwnerName)
+{
+	if (PlaceableItemPredicates.Remove(OwnerName))
+	{
+		PlaceableItemFilteringChanged.Broadcast();
+	}
+}
+
 void FPlacementModeModule::GetItemsForCategory(FName CategoryName, TArray<TSharedPtr<FPlaceableItem>>& OutItems) const
 {
 	const FPlacementCategory* Category = Categories.Find(CategoryName);
@@ -540,7 +560,10 @@ void FPlacementModeModule::GetItemsForCategory(FName CategoryName, TArray<TShare
 	{
 		for (auto& Pair : Category->Items)
 		{
-			OutItems.Add(Pair.Value);
+			if (PassesFilters(Pair.Value))
+			{
+				OutItems.Add(Pair.Value);
+			}
 		}
 	}
 }
@@ -552,9 +575,12 @@ void FPlacementModeModule::GetFilteredItemsForCategory(FName CategoryName, TArra
 	{
 		for (auto& Pair : Category->Items)
 		{
-			if (Filter(Pair.Value))
+			if (PassesFilters(Pair.Value))
 			{
-				OutItems.Add(Pair.Value);
+				if (Filter(Pair.Value))
+				{
+					OutItems.Add(Pair.Value);
+				}
 			}
 		}
 	}
@@ -580,9 +606,7 @@ void FPlacementModeModule::RegenerateItemsForCategory(FName Category)
 
 void FPlacementModeModule::RefreshRecentlyPlaced()
 {
-	FName CategoryName = FBuiltInPlacementCategories::RecentlyPlaced();
-
-	FPlacementCategory* Category = Categories.Find(CategoryName);
+	FPlacementCategory* Category = Categories.Find(FBuiltInPlacementCategories::RecentlyPlaced());
 	if (!Category)
 	{
 		return;
@@ -620,9 +644,7 @@ void FPlacementModeModule::RefreshRecentlyPlaced()
 
 void FPlacementModeModule::RefreshVolumes()
 {
-	FName CategoryName = FBuiltInPlacementCategories::Volumes();
-
-	FPlacementCategory* Category = Categories.Find(CategoryName);
+	FPlacementCategory* Category = Categories.Find(FBuiltInPlacementCategories::Volumes());
 	if (!Category)
 	{
 		return;
@@ -648,10 +670,8 @@ void FPlacementModeModule::RefreshVolumes()
 
 void FPlacementModeModule::RefreshAllPlaceableClasses()
 {
-	FName CategoryName = FBuiltInPlacementCategories::AllClasses();
-
 	// Unregister old stuff
-	FPlacementCategory* Category = Categories.Find(CategoryName);
+	FPlacementCategory* Category = Categories.Find(FBuiltInPlacementCategories::AllClasses());
 	if (!Category)
 	{
 		return;
@@ -728,6 +748,23 @@ FPlacementModeID FPlacementModeModule::CreateID(FName InCategory)
 	return NewID;
 }
 
+bool FPlacementModeModule::PassesFilters(const TSharedPtr<FPlaceableItem>& Item) const
+{
+	if (PlaceableItemPredicates.Num() == 0)
+	{
+		return true;
+	}
+
+	for (auto& PredicatePair : PlaceableItemPredicates)
+	{
+		if (PredicatePair.Value(Item))
+		{
+			return true;
+		}
+	}
+	return false;
+}
+
 void FPlacementModeModule::OnCategoryBlacklistChanged()
 {
 	PlacementModeCategoryListChanged.Broadcast();
diff --git a/Engine/Source/Editor/PlacementMode/Private/SPlacementModeTools.cpp b/Engine/Source/Editor/PlacementMode/Private/SPlacementModeTools.cpp
--- a/Engine/Source/Editor/PlacementMode/Private/SPlacementModeTools.cpp
+++ b/Engine/Source/Editor/PlacementMode/Private/SPlacementModeTools.cpp
@@ -22,6 +22,13 @@
 #include "Widgets/Docking/SDockTab.h"
 #include "ThumbnailRendering/ThumbnailManager.h"
 
+#define LOCTEXT_NAMESPACE "PlacementMode"
+
+namespace PlacementModeTools
+{
+	bool bItemInternalsInTooltip = false;
+	FAutoConsoleVariableRef CVarItemInternalsInTooltip(TEXT("PlacementMode.ItemInternalsInTooltip"), bItemInternalsInTooltip, TEXT("Shows placeable item internal information in its tooltip"));
+}
 
 struct FSortPlaceableItems
 {
@@ -143,19 +150,31 @@ void SPlacementAssetEntry::Construct(const FArguments& InArgs, const TSharedPtr<
 		DefaultActor = CastChecked<AActor>(CastChecked<UClass>(Item->AssetData.GetAsset())->ClassDefaultObject);
 	}
 
-	UClass* DocClass = nullptr;
 	TSharedPtr<IToolTip> AssetEntryToolTip;
+	if (PlacementModeTools::bItemInternalsInTooltip)
+	{
+		AssetEntryToolTip = FSlateApplicationBase::Get().MakeToolTip(
+			FText::Format(LOCTEXT("ItemInternalsTooltip", "Native Name: {0}\nAsset Path: {1}\nFactory Class: {2}"), 
+			FText::FromString(Item->NativeName), 
+			FText::FromName(Item->AssetData.ObjectPath),
+			FText::FromString(Item->Factory ? Item->Factory->GetClass()->GetName() : TEXT("None"))));
+	}
+
+	UClass* DocClass = nullptr;
 	if(DefaultActor != nullptr)
 	{
 		DocClass = DefaultActor->GetClass();
-		AssetEntryToolTip = FEditorClassUtils::GetTooltip(DefaultActor->GetClass());
+		if (!AssetEntryToolTip)
+		{
+			AssetEntryToolTip = FEditorClassUtils::GetTooltip(DefaultActor->GetClass());
+		}
 	}
 
-	if (!AssetEntryToolTip.IsValid())
+	if (!AssetEntryToolTip)
 	{
 		AssetEntryToolTip = FSlateApplicationBase::Get().MakeToolTip(Item->DisplayName);
 	}
-	
+
 	const FButtonStyle& ButtonStyle = FEditorStyle::GetWidgetStyle<FButtonStyle>( "PlacementBrowser.Asset" );
 
 	NormalImage = &ButtonStyle.Normal;
@@ -542,14 +561,16 @@ SPlacementModeTools::~SPlacementModeTools()
 		PlacementModeModule.OnRecentlyPlacedChanged().RemoveAll(this);
 		PlacementModeModule.OnAllPlaceableAssetsChanged().RemoveAll(this);
 		PlacementModeModule.OnPlacementModeCategoryListChanged().RemoveAll(this);
+		PlacementModeModule.OnPlaceableItemFilteringChanged().RemoveAll(this);
 	}
 }
 
 void SPlacementModeTools::Construct( const FArguments& InArgs, TSharedRef<SDockTab> ParentTab )
 {
-	bPlaceablesFullRefreshRequested = false;
-	bRecentlyPlacedRefreshRequested = false;
-	bNeedsUpdate = true;
+	bRefreshAllClasses = false;
+	bRefreshRecentlyPlaced = false;
+	bUpdateShownItems = true;
+
 	ActiveTabName = FBuiltInPlacementCategories::Basic();
 
 	ParentTab->SetOnTabDrawerOpened(FSimpleDelegate::CreateSP(this, &SPlacementModeTools::OnTabDrawerOpened));
@@ -579,7 +600,7 @@ void SPlacementModeTools::Construct( const FArguments& InArgs, TSharedRef<SDockT
 			.Padding(8)
 			[
 				SAssignNew( SearchBoxPtr, SSearchBox )
-				.HintText(NSLOCTEXT("PlacementMode", "SearchPlaceables", "Search Classes"))
+				.HintText(LOCTEXT("SearchPlaceables", "Search Classes"))
 				.OnTextChanged(this, &SPlacementModeTools::OnSearchChanged)
 				.OnTextCommitted(this, &SPlacementModeTools::OnSearchCommitted)
 			]
@@ -610,7 +631,7 @@ void SPlacementModeTools::Construct( const FArguments& InArgs, TSharedRef<SDockT
 			.Visibility(this, &SPlacementModeTools::GetTabsVisibility)
 			[
 				SAssignNew(FilterLabelPtr, STextBlock)
-				.Text(NSLOCTEXT("PlacementMode", "CategoryLabel", "CategoryLabel"))
+				.Text(LOCTEXT("CategoryLabel", "CategoryLabel"))
 				.Font(FAppStyle::Get().GetFontStyle("SmallFontBold"))
 				.TransformPolicy(ETextTransformPolicy::ToUpper)
 			]
@@ -627,7 +648,7 @@ void SPlacementModeTools::Construct( const FArguments& InArgs, TSharedRef<SDockT
 			.Padding(12.f)
 			[
 				SNew(STextBlock)
-				.Text(NSLOCTEXT("PlacementMode", "NoResultsFound", "No Results Found"))
+				.Text(LOCTEXT("NoResultsFound", "No Results Found"))
 				.Visibility(this, &SPlacementModeTools::GetFailedSearchVisibility)
 			]
 
@@ -661,8 +682,9 @@ void SPlacementModeTools::Construct( const FArguments& InArgs, TSharedRef<SDockT
 	];
 
 	IPlacementModeModule& PlacementModeModule = IPlacementModeModule::Get();
-	PlacementModeModule.OnRecentlyPlacedChanged().AddSP(this, &SPlacementModeTools::UpdateRecentlyPlacedAssets);
-	PlacementModeModule.OnAllPlaceableAssetsChanged().AddSP(this, &SPlacementModeTools::UpdatePlaceableAssets);
+	PlacementModeModule.OnRecentlyPlacedChanged().AddSP(this, &SPlacementModeTools::RequestRefreshRecentlyPlaced);
+	PlacementModeModule.OnAllPlaceableAssetsChanged().AddSP(this, &SPlacementModeTools::RequestRefreshAllClasses);
+	PlacementModeModule.OnPlaceableItemFilteringChanged().AddSP(this, &SPlacementModeTools::RequestUpdateShownItems);
 	PlacementModeModule.OnPlacementModeCategoryListChanged().AddSP(this, &SPlacementModeTools::UpdatePlacementCategories);
 }
 
@@ -677,13 +699,13 @@ void SPlacementModeTools::SetActiveTab(FName TabName)
 	{
 		ActiveTabName = TabName;
 		IPlacementModeModule::Get().RegenerateItemsForCategory(ActiveTabName);
-		bNeedsUpdate = true;
+		bUpdateShownItems = true;
 	}
 }
 
-void SPlacementModeTools::UpdateFilteredItems()
+void SPlacementModeTools::UpdateShownItems()
 {
-	bNeedsUpdate = false;
+	bUpdateShownItems = false;
 
 	IPlacementModeModule& PlacementModeModule = IPlacementModeModule::Get();
 
@@ -779,19 +801,24 @@ void SPlacementModeTools::OnTabDrawerOpened()
 	FSlateApplication::Get().SetKeyboardFocus(SearchBoxPtr, EFocusCause::SetDirectly);
 }
 
-void SPlacementModeTools::UpdateRecentlyPlacedAssets( const TArray< FActorPlacementInfo >& RecentlyPlaced )
+void SPlacementModeTools::RequestUpdateShownItems()
+{
+	bUpdateShownItems = true;
+}
+
+void SPlacementModeTools::RequestRefreshRecentlyPlaced( const TArray< FActorPlacementInfo >& RecentlyPlaced )
 {
 	if (GetActiveTab() == FBuiltInPlacementCategories::RecentlyPlaced())
 	{
-		bRecentlyPlacedRefreshRequested = true;
+		bRefreshRecentlyPlaced = true;
 	}
 }
 
-void SPlacementModeTools::UpdatePlaceableAssets()
+void SPlacementModeTools::RequestRefreshAllClasses()
 {
 	if (GetActiveTab() == FBuiltInPlacementCategories::AllClasses())
 	{
-		bPlaceablesFullRefreshRequested = true;
+		bRefreshAllClasses = true;
 	}
 }
 
@@ -803,7 +830,7 @@ void SPlacementModeTools::UpdatePlacementCategories()
 	CategoryFilterPtr->ClearChildren();
 
 	TArray<FPlacementCategoryInfo> Categories;
-	IPlacementModeModule::Get().GetUserFacingCategories(Categories);
+	IPlacementModeModule::Get().GetSortedCategories(Categories);
 	for (const FPlacementCategoryInfo& Category : Categories)
 	{
 		if (Category.UniqueHandle == FBuiltInPlacementCategories::Basic())
@@ -849,23 +876,23 @@ void SPlacementModeTools::UpdatePlacementCategories()
 
 void SPlacementModeTools::Tick( const FGeometry& AllottedGeometry, const double InCurrentTime, const float InDeltaTime )
 {
-	if ( bPlaceablesFullRefreshRequested )
+	if (bRefreshAllClasses)
 	{
 		IPlacementModeModule::Get().RegenerateItemsForCategory(FBuiltInPlacementCategories::AllClasses());
-		bPlaceablesFullRefreshRequested = false;
-		bNeedsUpdate = true;
+		bRefreshAllClasses = false;
+		bUpdateShownItems = true;
 	}
 
-	if ( bRecentlyPlacedRefreshRequested )
+	if (bRefreshRecentlyPlaced)
 	{
 		IPlacementModeModule::Get().RegenerateItemsForCategory(FBuiltInPlacementCategories::RecentlyPlaced());
-		bRecentlyPlacedRefreshRequested = false;
-		bNeedsUpdate = true;
+		bRefreshRecentlyPlaced = false;
+		bUpdateShownItems = true;
 	}
 
-	if ( bNeedsUpdate )
+	if (bUpdateShownItems)
 	{
-		UpdateFilteredItems();
+		UpdateShownItems();
 	}
 }
 
@@ -875,11 +902,11 @@ void SPlacementModeTools::OnSearchChanged(const FText& InFilterText)
 	// for the placeable widgets.
 	if ( !IsSearchActive() )
 	{
-		bPlaceablesFullRefreshRequested = true;
+		bRefreshAllClasses = true;
 	}
 	else
 	{
-		bNeedsUpdate = true;
+		bUpdateShownItems = true;
 	}
 
 	SearchTextFilter->SetRawFilterText( InFilterText );
@@ -895,3 +922,5 @@ FText SPlacementModeTools::GetHighlightText() const
 {
 	return SearchTextFilter->GetRawFilterText();
 }
+
+#undef LOCTEXT_NAMESPACE
\ No newline at end of file
diff --git a/Engine/Source/Editor/PlacementMode/Private/SPlacementModeTools.h b/Engine/Source/Editor/PlacementMode/Private/SPlacementModeTools.h
--- a/Engine/Source/Editor/PlacementMode/Private/SPlacementModeTools.h
+++ b/Engine/Source/Editor/PlacementMode/Private/SPlacementModeTools.h
@@ -111,16 +111,19 @@ private:
 	/** Check if a search is active */
 	bool IsSearchActive() const;
 
-	/** Update the list of filtered items */
-	void UpdateFilteredItems();
+	/** Update the list of shown items */
+	void UpdateShownItems();
 
-	/** Called when the recently placed assets changes. */
-	void UpdateRecentlyPlacedAssets( const TArray< FActorPlacementInfo >& RecentlyPlaced );
+	/** Requests an update of the shown items */
+	void RequestUpdateShownItems();
 
-	/** Called when the list of placeable assets changes. */
-	void UpdatePlaceableAssets();
+	/** Requests a refresh of the list of recently placed assets */
+	void RequestRefreshRecentlyPlaced( const TArray< FActorPlacementInfo >& RecentlyPlaced );
 
-	/** Called when the list of placement categories changes. */
+	/** Requests a refresh of placeable assets in the All Classes category */
+	void RequestRefreshAllClasses();
+
+	/** Refreshes the list of placement categories */
 	void UpdatePlacementCategories();
 
 	/** When the tab is clicked we adjust the check state, so that the right style is displayed. */
@@ -151,12 +154,12 @@ private:
 	FText GetHighlightText() const;
 
 private:
-	/** Flags to invoke updates of particular caregories on tick */
-	bool bRecentlyPlacedRefreshRequested;
-	bool bPlaceablesFullRefreshRequested;
+	/** Flags to refresh particular categories on tick */
+	bool bRefreshRecentlyPlaced;
+	bool bRefreshAllClasses;
 
-	/** Flag to indicate that we need to update the filtered items array */
-	bool bNeedsUpdate;
+	/** Flag to indicate that we need to update the list of shown items */
+	bool bUpdateShownItems;
 
 	// The text filter used to filter the classes
 	typedef TTextFilter<const FPlaceableItem&> FPlacementAssetEntryTextFilter;
diff --git a/Engine/Source/Editor/PlacementMode/Public/IPlacementModeModule.h b/Engine/Source/Editor/PlacementMode/Public/IPlacementModeModule.h
--- a/Engine/Source/Editor/PlacementMode/Public/IPlacementModeModule.h
+++ b/Engine/Source/Editor/PlacementMode/Public/IPlacementModeModule.h
@@ -90,20 +90,20 @@ struct FPlaceableItem
 	FPlaceableItem(UActorFactory* InFactory, const FAssetData& InAssetData, TOptional<int32> InSortOrder = TOptional<int32>())
 		: Factory(InFactory)
 		, AssetData(InAssetData)
+		, bAlwaysUseGenericThumbnail(false)
 		, SortOrder(InSortOrder)
 	{
-		bAlwaysUseGenericThumbnail = false;
-		AutoSetDisplayName();
+		AutoSetNativeAndDisplayName();
 	}
 
 	/** Constructor for any placeable class */
 	FPlaceableItem(UClass& InAssetClass, TOptional<int32> InSortOrder = TOptional<int32>())
 		: Factory(GEditor->FindActorFactoryByClass(&InAssetClass))
-		, AssetData(Factory ? Factory->GetDefaultActorClass( FAssetData() ) : FAssetData())
+		, AssetData(Factory ? Factory->GetDefaultActorClass(FAssetData()) : FAssetData())
+		, bAlwaysUseGenericThumbnail(false)
 		, SortOrder(InSortOrder)
 	{
-		bAlwaysUseGenericThumbnail = false;
-		AutoSetDisplayName();
+		AutoSetNativeAndDisplayName();
 	}
 
 	/** Constructor for any placeable class with associated asset data, brush and display name overrides */
@@ -120,22 +120,19 @@ struct FPlaceableItem
 		, AssetData(InAssetData)
 		, ClassThumbnailBrushOverride(InClassThumbnailBrushOverride)
 		, ClassIconBrushOverride(InClassIconBrushOverride)
+		, bAlwaysUseGenericThumbnail(true)
 		, AssetTypeColorOverride(InAssetTypeColorOverride)
 		, SortOrder(InSortOrder)
 	{
-		bAlwaysUseGenericThumbnail = true;
+		AutoSetNativeAndDisplayName();
 		if (InDisplayName.IsSet())
 		{
 			DisplayName = InDisplayName.GetValue();
 		}
-		else
-		{
-			AutoSetDisplayName();
-		}
 	}
 
-	/** Automatically set this item's display name from its class or asset */
-	void AutoSetDisplayName()
+	/** Automatically set this item's native and display names from its class or asset */
+	void AutoSetNativeAndDisplayName()
 	{
 		UClass* Class = AssetData.GetClass() == UClass::StaticClass() ? Cast<UClass>(AssetData.GetAsset()) : nullptr;
 		const bool bIsVolume = Class && Class->IsChildOf<AVolume>();
@@ -144,14 +141,18 @@ struct FPlaceableItem
 		// Use the factory unless its a volume or shape.  Those need custom names as the factory that spawns them does not properly represent what is being spawned.
 		if (Factory && !bIsVolume && !bIsShape) 
 		{
-			if (Factory->NewActorClassName.IsEmpty() == false)
+			if (!Factory->NewActorClassName.IsEmpty())
 			{
 				NativeName = Factory->NewActorClassName;
 			}
-			else if (Factory->NewActorClass != nullptr)
+			else if (Factory->NewActorClass)
 			{
 				Factory->NewActorClass->GetName(NativeName);
 			}
+			else
+			{
+				NativeName = Factory->GetName();
+			}
 
 			DisplayName = Factory->GetDisplayName();
 		}
@@ -162,10 +163,26 @@ struct FPlaceableItem
 		}
 		else
 		{
+			NativeName = AssetData.AssetName.ToString();
 			DisplayName = FText::FromName(AssetData.AssetName);
 		}
 	}
 
+	/** @deprecated Use AutoSetNativeAndDisplayName instead */
+	void AutoSetDisplayName() { AutoSetNativeAndDisplayName(); }
+
+	/** Return NativeName as an FName (and cache it) */
+	FName GetNativeFName() const
+	{ 
+		if (NativeFName.IsNone() && !NativeName.IsEmpty())
+		{
+			NativeFName = FName(*NativeName);
+		}
+		return NativeFName;
+	}
+
+public:
+
 	/** The factory used to create an instance of this placeable item */
 	UActorFactory* Factory;
 
@@ -192,6 +209,11 @@ struct FPlaceableItem
 
 	/** Optional sort order (lowest first). Overrides default class name sorting. */
 	TOptional<int32> SortOrder;
+
+private:
+
+	/** This item's native name as an FName (initialized on access only) */
+	mutable FName NativeFName;
 };
 
 /** Structure of built-in placement categories. Defined as functions to enable external use without linkage */
@@ -270,10 +292,17 @@ public:
 	DECLARE_EVENT( IPlacementModeModule, FOnAllPlaceableAssetsChanged );
 	virtual FOnAllPlaceableAssetsChanged& OnAllPlaceableAssetsChanged() = 0;
 
+	/**
+	 * @return the event that is broadcast whenever the filtering of placeable items changes (system filtering, not user filtering)
+	 */
+	DECLARE_EVENT(IPlacementMode, FOnPlaceableItemFilteringChanged);
+	virtual FOnPlaceableItemFilteringChanged& OnPlaceableItemFilteringChanged() = 0;
+
 	/**
 	 * Creates the placement browser widget
 	 */
 	virtual TSharedRef<SWidget> CreatePlacementModeBrowser(TSharedRef<SDockTab> ParentTab) = 0;
+
 public:
 
 	/**
@@ -303,16 +332,11 @@ public:
 	virtual TSharedRef<FBlacklistNames>& GetCategoryBlacklist() = 0;
 
 	/**
-	 * Populate the specified array with all registered category information that isn't blacklisted, sorted by SortOrder
+	 * Get all placement categories that aren't blacklisted, sorted by SortOrder
 	 *
-	 * @param OutCategories	The array to populate with registered category inforamtion
-	 */
-	virtual void GetUserFacingCategories(TArray<FPlacementCategoryInfo>& OutCategories) const = 0;
-
-	/**
-	 * @deprecated Use GetUserFacingCategories instead 
+	 * @param OutCategories	The array to populate with registered category information
 	 */
-	void GetSortedCategories(TArray<FPlacementCategoryInfo>& OutCategories) const { GetUserFacingCategories(OutCategories); }
+	virtual void GetSortedCategories(TArray<FPlacementCategoryInfo>& OutCategories) const = 0;
 
 	/**
 	 * Register a new placeable item for the specified category
@@ -330,8 +354,25 @@ public:
 	 */
 	virtual void UnregisterPlaceableItem(FPlacementModeID ID) = 0;
 
+	typedef TFunction<bool(const TSharedPtr<FPlaceableItem>&)> TPlaceableItemPredicate;
+
+	/** 
+	 * Registers system-level (not user) filtering for placeable items. 
+	 * An item is displayed if at least one of the predicate returns true or if there's none registered.
+	 * @param Predicate Function that returns true if the passed item should be available
+	 * @param OwnerName Name of the predicate owner
+	 * @return False on failure to register the predicate because one already exists under the specified owner name
+	 */
+	virtual bool RegisterPlaceableItemFilter(TPlaceableItemPredicate Predicate, FName OwnerName) = 0;
+
+	/**
+	 * Registers system-level (not user) filtering for placeable items.
+	 * An item is displayed if at least one of the predicate returns true or if there's none registered.
+	 */
+	virtual void UnregisterPlaceableItemFilter(FName OwnerName) = 0;
+
 	/**
-	 * Get all the items in a given category, unsorted
+	 * Get all items in a given category, system filtered, unsorted
 	 *
 	 * @param Category		The unique handle of the category to get items for
 	 * @param OutItems		Array to populate with the items in this category
@@ -339,7 +380,7 @@ public:
 	virtual void GetItemsForCategory(FName Category, TArray<TSharedPtr<FPlaceableItem>>& OutItems) const = 0;
 
 	/**
-	 * Get all the items in a given category, filtered by the specified predicate
+	 * Get all items in a given category, system and user filtered, unsorted
 	 *
 	 * @param Category		The unique handle of the category to get items for
 	 * @param OutItems		Array to populate with the items in this category
diff --git a/Engine/Source/Editor/PlacementMode/Public/PlacementModeModule.h b/Engine/Source/Editor/PlacementMode/Public/PlacementModeModule.h
--- a/Engine/Source/Editor/PlacementMode/Public/PlacementModeModule.h
+++ b/Engine/Source/Editor/PlacementMode/Public/PlacementModeModule.h
@@ -4,7 +4,6 @@
 
 #include "CoreMinimal.h"
 #include "IPlacementModeModule.h"
-
 #include "Framework/MultiBox/MultiBoxExtender.h"
 
 struct FPlacementCategory : FPlacementCategoryInfo
@@ -60,6 +59,9 @@ public:
 	DECLARE_DERIVED_EVENT(FPlacementModeModule, IPlacementModeModule::FOnAllPlaceableAssetsChanged, FOnAllPlaceableAssetsChanged);
 	virtual FOnAllPlaceableAssetsChanged& OnAllPlaceableAssetsChanged() override { return AllPlaceableAssetsChanged; }
 
+	DECLARE_DERIVED_EVENT(FPlacementModeModule, IPlacementModeModule::FOnPlaceableItemFilteringChanged, FOnPlaceableItemFilteringChanged);
+	virtual FOnPlaceableItemFilteringChanged& OnPlaceableItemFilteringChanged() override { return PlaceableItemFilteringChanged; };
+
 	/**
 	 * Add the specified assets to the recently placed items list
 	 */
@@ -91,12 +93,16 @@ public:
 
 	virtual TSharedRef<FBlacklistNames>& GetCategoryBlacklist() override { return CategoryBlacklist; }
 
-	virtual void GetUserFacingCategories(TArray<FPlacementCategoryInfo>& OutCategories) const;
+	virtual void GetSortedCategories(TArray<FPlacementCategoryInfo>& OutCategories) const;
 
 	virtual TOptional<FPlacementModeID> RegisterPlaceableItem(FName CategoryName, const TSharedRef<FPlaceableItem>& InItem);
 
 	virtual void UnregisterPlaceableItem(FPlacementModeID ID);
 
+	virtual bool RegisterPlaceableItemFilter(TPlaceableItemPredicate Predicate, FName OwnerName) override;
+
+	virtual void UnregisterPlaceableItemFilter(FName OwnerName) override;
+
 	virtual void GetItemsForCategory(FName CategoryName, TArray<TSharedPtr<FPlaceableItem>>& OutItems) const;
 
 	virtual void GetFilteredItemsForCategory(FName CategoryName, TArray<TSharedPtr<FPlaceableItem>>& OutItems, TFunctionRef<bool(const TSharedPtr<FPlaceableItem>&)> Filter) const;
@@ -106,23 +112,19 @@ public:
 private:
 
 	void OnAssetRemoved(const FAssetData& /*InRemovedAssetData*/);
-
 	void OnAssetRenamed(const FAssetData& AssetData, const FString& OldObjectPath);
-
 	void OnAssetAdded(const FAssetData& AssetData);
-
 	void OnInitialAssetsScanComplete();
 
 	void RefreshRecentlyPlaced();
-
 	void RefreshVolumes();
-
 	void RefreshAllPlaceableClasses();
 
 	FGuid CreateID();
-
 	FPlacementModeID CreateID(FName InCategory);
 
+	bool PassesFilters(const TSharedPtr<FPlaceableItem>& Item) const;
+
 	void OnCategoryBlacklistChanged();
 
 private:
@@ -131,11 +133,14 @@ private:
 
 	TSharedRef<FBlacklistNames> CategoryBlacklist;
 
+	TMap<FName, TPlaceableItemPredicate> PlaceableItemPredicates;
+
 	TArray< FActorPlacementInfo > RecentlyPlaced;
 	FOnRecentlyPlacedChanged RecentlyPlacedChanged;
 
 	FOnAllPlaceableAssetsChanged AllPlaceableAssetsChanged;
 	FOnPlacementModeCategoryRefreshed PlacementModeCategoryRefreshed;
+	FOnPlaceableItemFilteringChanged PlaceableItemFilteringChanged;
 	FOnPlacementModeCategoryListChanged PlacementModeCategoryListChanged;
 
 	TArray< TSharedPtr<FExtender> > ContentPaletteFiltersExtenders;
