Fix PlacementModeModule event declaration

diff --git a/Engine/Source/Editor/PlacementMode/Public/IPlacementModeModule.h b/Engine/Source/Editor/PlacementMode/Public/IPlacementModeModule.h
--- a/Engine/Source/Editor/PlacementMode/Public/IPlacementModeModule.h
+++ b/Engine/Source/Editor/PlacementMode/Public/IPlacementModeModule.h
@@ -271,7 +271,7 @@ public:
 	/**
 	 * @return the event that is broadcast whenever the user facing list of placement mode categories gets modified
 	 */
-	DECLARE_EVENT(IPlacementMode, FOnPlacementModeCategoryListChanged);
+	DECLARE_EVENT(IPlacementModeModule, FOnPlacementModeCategoryListChanged);
 	virtual FOnPlacementModeCategoryListChanged& OnPlacementModeCategoryListChanged() = 0;
 
 	/**
@@ -295,7 +295,7 @@ public:
 	/**
 	 * @return the event that is broadcast whenever the filtering of placeable items changes (system filtering, not user filtering)
 	 */
-	DECLARE_EVENT(IPlacementMode, FOnPlaceableItemFilteringChanged);
+	DECLARE_EVENT(IPlacementModeModule, FOnPlaceableItemFilteringChanged);
 	virtual FOnPlaceableItemFilteringChanged& OnPlaceableItemFilteringChanged() = 0;
 
 	/**
