Fixed Assert after clicking Create>Cinematics or Create>Testing

diff --git a/Engine/Source/Editor/PlacementMode/Private/SPlacementModeTools.cpp b/Engine/Source/Editor/PlacementMode/Private/SPlacementModeTools.cpp
--- a/Engine/Source/Editor/PlacementMode/Private/SPlacementModeTools.cpp
+++ b/Engine/Source/Editor/PlacementMode/Private/SPlacementModeTools.cpp
@@ -473,8 +473,20 @@ FReply SPlacementAssetMenuEntry::OnMouseButtonUp(const FGeometry& MyGeometry, co
 	{
 		bIsPressed = false;
 
+		AActor* NewActor = nullptr;
 		FTransform TempTransform;
-		if( AActor* NewActor = FLevelEditorActionCallbacks::AddActor(Item->Factory, Item->AssetData, &TempTransform) )
+		if (Item->Factory != nullptr)
+		{
+			NewActor = FLevelEditorActionCallbacks::AddActor(Item->Factory, Item->AssetData, &TempTransform);
+		}
+		else 
+		{
+			// If no actor factory was found or failed, add the actor from the uclass.
+			AActor* DefaultActor = CastChecked<AActor>(CastChecked<UClass>(Item->AssetData.GetAsset())->ClassDefaultObject);
+			NewActor = FLevelEditorActionCallbacks::AddActorFromClass(DefaultActor->GetClass());
+		}
+
+		if (NewActor != nullptr)
 		{
   			GEditor->MoveActorInFrontOfCamera(*NewActor, 
   				GCurrentLevelEditingViewportClient->GetViewLocation(), 
