Fixes for create menu - handling of missing factories fixed to prevent crashes, items checked for if they are placeable

diff --git a/Engine/Source/Editor/PlacementMode/Private/PlacementModeModule.cpp b/Engine/Source/Editor/PlacementMode/Private/PlacementModeModule.cpp
--- a/Engine/Source/Editor/PlacementMode/Private/PlacementModeModule.cpp
+++ b/Engine/Source/Editor/PlacementMode/Private/PlacementModeModule.cpp
@@ -52,6 +52,7 @@
 #include "ActorFactories/ActorFactoryPlanarReflection.h"
 #include "SPlacementModeTools.h"
 #include "Classes/EditorStyleSettings.h"
+#include "AssetSelection.h"
 
 
 TOptional<FLinearColor> GetBasicShapeColorOverride()
@@ -627,14 +628,17 @@ void FPlacementModeModule::RefreshRecentlyPlaced()
 			if (AssetData.IsValid())
 			{
 				UActorFactory* Factory = FindObject<UActorFactory>(nullptr, *RecentlyPlacedItem.Factory);
-				TSharedPtr<FPlaceableItem> Ptr = MakeShareable(new FPlaceableItem(Factory, AssetData));
-				if (FString* FoundThumbnail = BasicShapeThumbnails.Find(RecentlyPlacedItem.ObjectPath))
+				if (Factory)
 				{
-					Ptr->ClassThumbnailBrushOverride = FName(**FoundThumbnail);
-					Ptr->bAlwaysUseGenericThumbnail = true;
-					Ptr->AssetTypeColorOverride = GetBasicShapeColorOverride();
+					TSharedPtr<FPlaceableItem> Ptr = MakeShareable(new FPlaceableItem(Factory, AssetData));
+					if (FString* FoundThumbnail = BasicShapeThumbnails.Find(RecentlyPlacedItem.ObjectPath))
+					{
+						Ptr->ClassThumbnailBrushOverride = FName(**FoundThumbnail);
+						Ptr->bAlwaysUseGenericThumbnail = true;
+						Ptr->AssetTypeColorOverride = GetBasicShapeColorOverride();
+					}
+					Category->Items.Add(CreateID(), Ptr);
 				}
-				Category->Items.Add(CreateID(), Ptr);
 			}
 		}
 	}
@@ -757,7 +761,41 @@ bool FPlacementModeModule::PassesFilters(const TSharedPtr<FPlaceableItem>& Item)
 	{
 		if (PredicatePair.Value(Item))
 		{
-			return true;
+			bool bPlaceable = true;
+			if (Item->AssetData.GetClass() == UClass::StaticClass())
+			{
+				UClass* Class = Cast<UClass>(Item->AssetData.GetAsset());
+
+				bPlaceable = AssetSelectionUtils::IsClassPlaceable(Class);
+			}
+			else if (Item->AssetData.GetClass()->IsChildOf<UBlueprint>())
+			{
+				// For blueprints, attempt to determine placeability from its tag information
+
+				FString TagValue;
+
+				if (Item->AssetData.GetTagValue(FBlueprintTags::NativeParentClassPath, TagValue) && !TagValue.IsEmpty())
+				{
+					// If the native parent class can't be placed, neither can the blueprint
+
+					UObject* Outer = nullptr;
+					ResolveName(Outer, TagValue, false, false);
+					UClass* NativeParentClass = FindObject<UClass>(ANY_PACKAGE, *TagValue);
+
+					bPlaceable = AssetSelectionUtils::IsChildBlueprintPlaceable(NativeParentClass);
+				}
+
+				if (bPlaceable && Item->AssetData.GetTagValue(FBlueprintTags::ClassFlags, TagValue) && !TagValue.IsEmpty())
+				{
+					// Check to see if this class is placeable from its class flags
+
+					const int32 NotPlaceableFlags = CLASS_NotPlaceable | CLASS_Deprecated | CLASS_Abstract;
+					uint32 ClassFlags = FCString::Atoi(*TagValue);
+
+					bPlaceable = (ClassFlags & NotPlaceableFlags) == CLASS_None;
+				}
+			}
+			return bPlaceable;
 		}
 	}
 	return false;
diff --git a/Engine/Source/Editor/PlacementMode/Private/SPlacementModeTools.cpp b/Engine/Source/Editor/PlacementMode/Private/SPlacementModeTools.cpp
--- a/Engine/Source/Editor/PlacementMode/Private/SPlacementModeTools.cpp
+++ b/Engine/Source/Editor/PlacementMode/Private/SPlacementModeTools.cpp
@@ -21,6 +21,8 @@
 #include "ClassIconFinder.h"
 #include "Widgets/Docking/SDockTab.h"
 #include "ThumbnailRendering/ThumbnailManager.h"
+#include "AssetSelection.h"
+#include "ActorFactories/ActorFactory.h"
 
 #define LOCTEXT_NAMESPACE "PlacementMode"
 
@@ -475,17 +477,14 @@ FReply SPlacementAssetMenuEntry::OnMouseButtonUp(const FGeometry& MyGeometry, co
 
 		AActor* NewActor = nullptr;
 		FTransform TempTransform;
-		if (Item->Factory != nullptr)
+		UActorFactory* Factory = Item->Factory;
+		if (!Item->Factory)
 		{
-			NewActor = FLevelEditorActionCallbacks::AddActor(Item->Factory, Item->AssetData, &TempTransform);
+			// If no actor factory was found or failed, add the actor from the uclass
+			UObject* ClassObject = Item->AssetData.GetClass()->GetDefaultObject();
+			FActorFactoryAssetProxy::GetFactoryForAssetObject(ClassObject);
 		}
-		else 
-		{
-			// If no actor factory was found or failed, add the actor from the uclass.
-			AActor* DefaultActor = CastChecked<AActor>(CastChecked<UClass>(Item->AssetData.GetAsset())->ClassDefaultObject);
-			NewActor = FLevelEditorActionCallbacks::AddActorFromClass(DefaultActor->GetClass());
-		}
-
+		NewActor = FLevelEditorActionCallbacks::AddActor(Factory, Item->AssetData, &TempTransform);
 		if (NewActor != nullptr)
 		{
   			GEditor->MoveActorInFrontOfCamera(*NewActor, 
