Fixed a crash when creating an actor while there is no viewport open in the level editor.

diff --git a/Engine/Source/Editor/PlacementMode/Private/SPlacementModeTools.cpp b/Engine/Source/Editor/PlacementMode/Private/SPlacementModeTools.cpp
--- a/Engine/Source/Editor/PlacementMode/Private/SPlacementModeTools.cpp
+++ b/Engine/Source/Editor/PlacementMode/Private/SPlacementModeTools.cpp
@@ -485,7 +485,7 @@ FReply SPlacementAssetMenuEntry::OnMouseButtonUp(const FGeometry& MyGeometry, co
 			FActorFactoryAssetProxy::GetFactoryForAssetObject(ClassObject);
 		}
 		NewActor = FLevelEditorActionCallbacks::AddActor(Factory, Item->AssetData, &TempTransform);
-		if (NewActor != nullptr)
+		if (NewActor && GCurrentLevelEditingViewportClient)
 		{
   			GEditor->MoveActorInFrontOfCamera(*NewActor, 
   				GCurrentLevelEditingViewportClient->GetViewLocation(), 
