The Create menu entry should pull its name from the asset vs. the factory in most cases

diff --git a/Engine/Source/Editor/PlacementMode/Public/IPlacementModeModule.h b/Engine/Source/Editor/PlacementMode/Public/IPlacementModeModule.h
--- a/Engine/Source/Editor/PlacementMode/Public/IPlacementModeModule.h
+++ b/Engine/Source/Editor/PlacementMode/Public/IPlacementModeModule.h
@@ -138,25 +138,8 @@ struct FPlaceableItem
 		const bool bIsVolume = Class && Class->IsChildOf<AVolume>();
 		const bool bIsShape = Class ? false : AssetData.GetClass()->IsChildOf(UStaticMesh::StaticClass());
 
-		// Use the factory unless its a volume or shape.  Those need custom names as the factory that spawns them does not properly represent what is being spawned.
-		if (Factory && !bIsVolume && !bIsShape) 
-		{
-			if (!Factory->NewActorClassName.IsEmpty())
-			{
-				NativeName = Factory->NewActorClassName;
-			}
-			else if (Factory->NewActorClass)
-			{
-				Factory->NewActorClass->GetName(NativeName);
-			}
-			else
-			{
-				NativeName = Factory->GetName();
-			}
-
-			DisplayName = Factory->GetDisplayName();
-		}
-		else if (Class)
+		// Use the class unless its a volume or shape.  Those need custom names as the factory that spawns them does not properly represent what is being spawned.
+		if (Class && !bIsVolume && !bIsShape)
 		{
 			Class->GetName(NativeName);
 			DisplayName = Class->GetDisplayNameText();
