Sort "All Actors" in create menu alphabetically

diff --git a/Engine/Source/Editor/PlacementMode/Private/PlacementModeModule.cpp b/Engine/Source/Editor/PlacementMode/Private/PlacementModeModule.cpp
--- a/Engine/Source/Editor/PlacementMode/Private/PlacementModeModule.cpp
+++ b/Engine/Source/Editor/PlacementMode/Private/PlacementModeModule.cpp
@@ -735,6 +735,10 @@ void FPlacementModeModule::RefreshAllPlaceableClasses()
 			Category->Items.Add(CreateID(), MakeShareable(new FPlaceableItem(ActorFactory, FAssetData(*ClassIt))));
 		}
 	}
+
+	Category->Items.ValueSort([&](const TSharedPtr<FPlaceableItem>& A, const TSharedPtr<FPlaceableItem>& B) {
+		return A->DisplayName.CompareTo(B->DisplayName) < 0;
+		});
 }
 
 FGuid FPlacementModeModule::CreateID()
