When adding a new Box Reflection Capture from the Quick Add Menu, the object is transformed once again to (1, 1, 1) after setting the right scale. The new code restores object to the original scale.

diff --git a/Engine/Source/Editor/PlacementMode/Private/SPlacementModeTools.cpp b/Engine/Source/Editor/PlacementMode/Private/SPlacementModeTools.cpp
--- a/Engine/Source/Editor/PlacementMode/Private/SPlacementModeTools.cpp
+++ b/Engine/Source/Editor/PlacementMode/Private/SPlacementModeTools.cpp
@@ -477,7 +477,6 @@ FReply SPlacementAssetMenuEntry::OnMouseButtonUp(const FGeometry& MyGeometry, co
 		bIsPressed = false;
 
 		AActor* NewActor = nullptr;
-		FTransform TempTransform;
 		UActorFactory* Factory = Item->Factory;
 		if (!Item->Factory)
 		{
@@ -489,7 +488,7 @@ FReply SPlacementAssetMenuEntry::OnMouseButtonUp(const FGeometry& MyGeometry, co
 				FActorFactoryAssetProxy::GetFactoryForAssetObject(ClassObject);
 			}
 		}
-		NewActor = FLevelEditorActionCallbacks::AddActor(Factory, Item->AssetData, &TempTransform);
+		NewActor = FLevelEditorActionCallbacks::AddActor(Factory, Item->AssetData, nullptr);
 		if (NewActor && GCurrentLevelEditingViewportClient)
 		{
   			GEditor->MoveActorInFrontOfCamera(*NewActor, 
