Fix crash with PlacementModeModule when editor closes

diff --git a/Engine/Source/Editor/PlacementMode/Private/PlacementModeModule.cpp b/Engine/Source/Editor/PlacementMode/Private/PlacementModeModule.cpp
--- a/Engine/Source/Editor/PlacementMode/Private/PlacementModeModule.cpp
+++ b/Engine/Source/Editor/PlacementMode/Private/PlacementModeModule.cpp
@@ -488,7 +488,7 @@ bool FPlacementModeModule::RegisterPlacementCategory(const FPlacementCategoryInf
 
 void FPlacementModeModule::UnregisterPlacementCategory(FName Handle)
 {
-	if (Categories.Remove(Handle))
+	if (Categories.Remove(Handle) && !IsEngineExitRequested())
 	{
 		PlacementModeCategoryListChanged.Broadcast();
 	}
