Fix to Placement Mode localization of Display Name to use the asset class name.

diff --git a/Engine/Source/Editor/PlacementMode/Public/IPlacementModeModule.h b/Engine/Source/Editor/PlacementMode/Public/IPlacementModeModule.h
--- a/Engine/Source/Editor/PlacementMode/Public/IPlacementModeModule.h
+++ b/Engine/Source/Editor/PlacementMode/Public/IPlacementModeModule.h
@@ -136,11 +136,8 @@ struct FPlaceableItem
 	void AutoSetNativeAndDisplayName()
 	{
 		UClass* Class = AssetData.GetClass() == UClass::StaticClass() ? Cast<UClass>(AssetData.GetAsset()) : nullptr;
-		const bool bIsVolume = Class && Class->IsChildOf<AVolume>();
-		const bool bIsShape = Class ? false : AssetData.IsInstanceOf(UStaticMesh::StaticClass());
 
-		// Use the class unless its a volume or shape.  Those need custom names as the factory that spawns them does not properly represent what is being spawned.
-		if (Class && !bIsVolume && !bIsShape)
+		if (Class)
 		{
 			Class->GetName(NativeName);
 			DisplayName = Class->GetDisplayNameText();
