Fix to PlacementMode localization. Instead of recreating Placement Items with partial information we keep the manually created Basic shapes in a map so we can fetch them when building the recently used list.

diff --git a/Engine/Source/Editor/PlacementMode/Private/PlacementModeModule.cpp b/Engine/Source/Editor/PlacementMode/Private/PlacementModeModule.cpp
--- a/Engine/Source/Editor/PlacementMode/Private/PlacementModeModule.cpp
+++ b/Engine/Source/Editor/PlacementMode/Private/PlacementModeModule.cpp
@@ -175,17 +175,43 @@ void FPlacementModeModule::StartupModule()
 				)
 			);
 
+		static const FText CubeText = NSLOCTEXT("PlacementMode", "Cube", "Cube");
+		static const FText SphereText = NSLOCTEXT("PlacementMode", "Sphere", "Sphere");
+		static const FText CylinderText = NSLOCTEXT("PlacementMode", "Cylinder", "Cylinder");
+		static const FText ConeText = NSLOCTEXT("PlacementMode", "Cone", "Cone");
+		static const FText PlaneText = NSLOCTEXT("PlacementMode", "Plane", "Plane");
+
 		FPlacementCategory* Category = Categories.Find(CategoryName);
 		// Cube
-		Category->Items.Add(CreateID(), MakeShareable(new FPlaceableItem(*UActorFactoryBasicShape::StaticClass(), FAssetData(LoadObject<UStaticMesh>(nullptr, *UActorFactoryBasicShape::BasicCube.ToString())), FName("ClassThumbnail.Cube"), FName("ClassIcon.Cube"), BasicShapeColorOverride, SortOrder += 10, NSLOCTEXT("PlacementMode", "Cube", "Cube"))));
+		{
+			TSharedPtr<FPlaceableItem> Cube = MakeShareable(new FPlaceableItem(*UActorFactoryBasicShape::StaticClass(), FAssetData(LoadObject<UStaticMesh>(nullptr, *UActorFactoryBasicShape::BasicCube.ToString())), FName("ClassThumbnail.Cube"), FName("ClassIcon.Cube"), BasicShapeColorOverride, SortOrder += 10, CubeText));
+			ManuallyCreatedPlaceableItems.Add(*UActorFactoryBasicShape::BasicCube.ToString(),Cube);
+			Category->Items.Add(CreateID(), Cube);
+		}
 		// Sphere
-		Category->Items.Add(CreateID(), MakeShareable(new FPlaceableItem(*UActorFactoryBasicShape::StaticClass(), FAssetData(LoadObject<UStaticMesh>(nullptr, *UActorFactoryBasicShape::BasicSphere.ToString())), FName("ClassThumbnail.Sphere"), FName("ClassIcon.Sphere"), BasicShapeColorOverride, SortOrder += 10, NSLOCTEXT("PlacementMode", "Sphere", "Sphere"))));
+		{
+			TSharedPtr<FPlaceableItem> Sphere = MakeShareable(new FPlaceableItem(*UActorFactoryBasicShape::StaticClass(), FAssetData(LoadObject<UStaticMesh>(nullptr, *UActorFactoryBasicShape::BasicSphere.ToString())), FName("ClassThumbnail.Sphere"), FName("ClassIcon.Sphere"), BasicShapeColorOverride, SortOrder += 10, SphereText));
+			ManuallyCreatedPlaceableItems.Add(*UActorFactoryBasicShape::BasicSphere.ToString(), Sphere);
+			Category->Items.Add(CreateID(), Sphere);
+		}
 		// Cylinder
-		Category->Items.Add(CreateID(), MakeShareable(new FPlaceableItem(*UActorFactoryBasicShape::StaticClass(), FAssetData(LoadObject<UStaticMesh>(nullptr, *UActorFactoryBasicShape::BasicCylinder.ToString())), FName("ClassThumbnail.Cylinder"), FName("ClassIcon.Cylinder"), BasicShapeColorOverride, SortOrder += 10, NSLOCTEXT("PlacementMode", "Cylinder", "Cylinder"))));
+		{
+			TSharedPtr<FPlaceableItem> Cylinder = MakeShareable(new FPlaceableItem(*UActorFactoryBasicShape::StaticClass(), FAssetData(LoadObject<UStaticMesh>(nullptr, *UActorFactoryBasicShape::BasicCylinder.ToString())), FName("ClassThumbnail.Cylinder"), FName("ClassIcon.Cylinder"), BasicShapeColorOverride, SortOrder += 10, CylinderText));
+			ManuallyCreatedPlaceableItems.Add(*UActorFactoryBasicShape::BasicCylinder.ToString(), Cylinder);
+			Category->Items.Add(CreateID(), Cylinder);
+		}
 		// Cone
-		Category->Items.Add(CreateID(), MakeShareable(new FPlaceableItem(*UActorFactoryBasicShape::StaticClass(), FAssetData(LoadObject<UStaticMesh>(nullptr, *UActorFactoryBasicShape::BasicCone.ToString())), FName("ClassThumbnail.Cone"), FName("ClassIcon.Cone"), BasicShapeColorOverride, SortOrder += 10, NSLOCTEXT("PlacementMode", "Cone", "Cone"))));
+		{
+			TSharedPtr<FPlaceableItem> Cone = MakeShareable(new FPlaceableItem(*UActorFactoryBasicShape::StaticClass(), FAssetData(LoadObject<UStaticMesh>(nullptr, *UActorFactoryBasicShape::BasicCone.ToString())), FName("ClassThumbnail.Cone"), FName("ClassIcon.Cone"), BasicShapeColorOverride, SortOrder += 10, ConeText));
+			ManuallyCreatedPlaceableItems.Add(*UActorFactoryBasicShape::BasicCone.ToString(),Cone);
+			Category->Items.Add(CreateID(), Cone);
+		}
 		// Plane
-		Category->Items.Add(CreateID(), MakeShareable(new FPlaceableItem(*UActorFactoryBasicShape::StaticClass(), FAssetData(LoadObject<UStaticMesh>(nullptr, *UActorFactoryBasicShape::BasicPlane.ToString())), FName("ClassThumbnail.Plane"), FName("ClassIcon.Plane"), BasicShapeColorOverride, SortOrder += 10, NSLOCTEXT("PlacementMode", "Plane", "Plane"))));
+		{
+			TSharedPtr<FPlaceableItem> Plane = MakeShareable(new FPlaceableItem(*UActorFactoryBasicShape::StaticClass(), FAssetData(LoadObject<UStaticMesh>(nullptr, *UActorFactoryBasicShape::BasicPlane.ToString())), FName("ClassThumbnail.Plane"), FName("ClassIcon.Plane"), BasicShapeColorOverride, SortOrder += 10, PlaneText));
+			ManuallyCreatedPlaceableItems.Add(*UActorFactoryBasicShape::BasicPlane.ToString(), Plane);
+			Category->Items.Add(CreateID(), Plane);
+		}
 	}
 
 	{
@@ -235,14 +261,6 @@ void FPlacementModeModule::StartupModule()
 		)
 	);
 
-
-	BasicShapeThumbnails.Add(UActorFactoryBasicShape::BasicCube.ToString(), TEXT("ClassThumbnail.Cube"));
-	BasicShapeThumbnails.Add(UActorFactoryBasicShape::BasicSphere.ToString(), TEXT("ClassThumbnail.Sphere"));
-	BasicShapeThumbnails.Add(UActorFactoryBasicShape::BasicCylinder.ToString(), TEXT("ClassThumbnail.Cylinder"));
-	BasicShapeThumbnails.Add(UActorFactoryBasicShape::BasicCone.ToString(), TEXT("ClassThumbnail.Cone"));
-	BasicShapeThumbnails.Add(UActorFactoryBasicShape::BasicPlane.ToString(), TEXT("ClassThumbnail.Plane"));
-
-
 	if (FSlateApplication::IsInitialized())
 	{
 		// Given a Category Name, this will add a section and all the placeable items in that section directly into the menu
@@ -622,26 +640,27 @@ void FPlacementModeModule::RefreshRecentlyPlaced()
 
 	for (const FActorPlacementInfo& RecentlyPlacedItem : RecentlyPlaced)
 	{
-		UObject* Asset = FindObject<UObject>(nullptr, *RecentlyPlacedItem.ObjectPath);
-
-		// If asset is pending delete, it will not be marked as RF_Standalone, in which case we skip it
-		if (Asset != nullptr && Asset->HasAnyFlags(RF_Standalone))
+		// First check if it's a manually created entry
+		if (TSharedPtr<FPlaceableItem>* ItemFound = ManuallyCreatedPlaceableItems.Find(RecentlyPlacedItem.ObjectPath))
+		{
+			Category->Items.Add(CreateID(), *ItemFound);
+		}
+		else
 		{
-			FAssetData AssetData = AssetRegistryModule.Get().GetAssetByObjectPath(FSoftObjectPath(RecentlyPlacedItem.ObjectPath));
+			UObject* Asset = FindObject<UObject>(nullptr, *RecentlyPlacedItem.ObjectPath);
 
-			if (AssetData.IsValid())
+			// If asset is pending delete, it will not be marked as RF_Standalone, in which case we skip it
+			if (Asset != nullptr && Asset->HasAnyFlags(RF_Standalone))
 			{
-				UActorFactory* Factory = FindObject<UActorFactory>(nullptr, *RecentlyPlacedItem.Factory);
-				if (Factory)
+				FAssetData AssetData = AssetRegistryModule.Get().GetAssetByObjectPath(FSoftObjectPath(RecentlyPlacedItem.ObjectPath));
+
+				if (AssetData.IsValid())
 				{
-					TSharedPtr<FPlaceableItem> Ptr = MakeShareable(new FPlaceableItem(Factory, AssetData));
-					if (FString* FoundThumbnail = BasicShapeThumbnails.Find(RecentlyPlacedItem.ObjectPath))
+					if (UActorFactory* Factory = FindObject<UActorFactory>(nullptr, *RecentlyPlacedItem.Factory))
 					{
-						Ptr->ClassThumbnailBrushOverride = FName(**FoundThumbnail);
-						Ptr->bAlwaysUseGenericThumbnail = true;
-						Ptr->AssetTypeColorOverride = GetBasicShapeColorOverride();
+						TSharedPtr<FPlaceableItem> Ptr = MakeShareable(new FPlaceableItem(Factory, AssetData));
+						Category->Items.Add(CreateID(), Ptr);
 					}
-					Category->Items.Add(CreateID(), Ptr);
 				}
 			}
 		}
diff --git a/Engine/Source/Editor/PlacementMode/Public/PlacementModeModule.h b/Engine/Source/Editor/PlacementMode/Public/PlacementModeModule.h
--- a/Engine/Source/Editor/PlacementMode/Public/PlacementModeModule.h
+++ b/Engine/Source/Editor/PlacementMode/Public/PlacementModeModule.h
@@ -146,7 +146,7 @@ private:
 	TArray< TSharedPtr<FExtender> > ContentPaletteFiltersExtenders;
 	TArray< TSharedPtr<FExtender> > PaletteExtenders;
 
-	TMap<FString, FString> BasicShapeThumbnails;
+	TMap<FString, TSharedPtr<FPlaceableItem>> ManuallyCreatedPlaceableItems;
 };
 
 IMPLEMENT_MODULE(FPlacementModeModule, PlacementMode);
\ No newline at end of file
