Editor: Make placement mode panel update itself if external code requests a category refresh.

diff --git a/Engine/Source/Editor/PlacementMode/Private/SPlacementModeTools.cpp b/Engine/Source/Editor/PlacementMode/Private/SPlacementModeTools.cpp
--- a/Engine/Source/Editor/PlacementMode/Private/SPlacementModeTools.cpp
+++ b/Engine/Source/Editor/PlacementMode/Private/SPlacementModeTools.cpp
@@ -701,6 +701,7 @@ void SPlacementModeTools::Construct( const FArguments& InArgs, TSharedRef<SDockT
 	PlacementModeModule.OnAllPlaceableAssetsChanged().AddSP(this, &SPlacementModeTools::RequestRefreshAllClasses);
 	PlacementModeModule.OnPlaceableItemFilteringChanged().AddSP(this, &SPlacementModeTools::RequestUpdateShownItems);
 	PlacementModeModule.OnPlacementModeCategoryListChanged().AddSP(this, &SPlacementModeTools::UpdatePlacementCategories);
+	PlacementModeModule.OnPlacementModeCategoryRefreshed().AddSP(this, &SPlacementModeTools::OnCategoryRefresh);
 }
 
 FName SPlacementModeTools::GetActiveTab() const
@@ -714,7 +715,6 @@ void SPlacementModeTools::SetActiveTab(FName TabName)
 	{
 		ActiveTabName = TabName;
 		IPlacementModeModule::Get().RegenerateItemsForCategory(ActiveTabName);
-		bUpdateShownItems = true;
 	}
 }
 
@@ -837,6 +837,14 @@ void SPlacementModeTools::RequestRefreshAllClasses()
 	}
 }
 
+void SPlacementModeTools::OnCategoryRefresh(FName CategoryName)
+{
+	if (GetActiveTab() == CategoryName)
+	{
+		RequestUpdateShownItems();
+	}
+}
+
 void SPlacementModeTools::UpdatePlacementCategories()
 {
 	bool BasicTabExists = false;
@@ -895,14 +903,12 @@ void SPlacementModeTools::Tick( const FGeometry& AllottedGeometry, const double
 	{
 		IPlacementModeModule::Get().RegenerateItemsForCategory(FBuiltInPlacementCategories::AllClasses());
 		bRefreshAllClasses = false;
-		bUpdateShownItems = true;
 	}
 
 	if (bRefreshRecentlyPlaced)
 	{
 		IPlacementModeModule::Get().RegenerateItemsForCategory(FBuiltInPlacementCategories::RecentlyPlaced());
 		bRefreshRecentlyPlaced = false;
-		bUpdateShownItems = true;
 	}
 
 	if (bUpdateShownItems)
diff --git a/Engine/Source/Editor/PlacementMode/Private/SPlacementModeTools.h b/Engine/Source/Editor/PlacementMode/Private/SPlacementModeTools.h
--- a/Engine/Source/Editor/PlacementMode/Private/SPlacementModeTools.h
+++ b/Engine/Source/Editor/PlacementMode/Private/SPlacementModeTools.h
@@ -125,6 +125,9 @@ private:
 	/** Requests a refresh of placeable assets in the All Classes category */
 	void RequestRefreshAllClasses();
 
+	/** Called when the placement module performs a category refresh */
+	void OnCategoryRefresh(FName CategoryName);
+
 	/** Refreshes the list of placement categories */
 	void UpdatePlacementCategories();
 
