Fixed placing an actor via the create menu failing to capture the initial transform correctly within a transaction diff

This was due to the actor being spawned at the default transform (within a transaction), and then moved post-spawn (outside of a transaction). This change has the entire spawn and move captured by a single transaction.

#rnx

diff --git a/Engine/Source/Editor/PlacementMode/Private/SPlacementModeTools.cpp b/Engine/Source/Editor/PlacementMode/Private/SPlacementModeTools.cpp
--- a/Engine/Source/Editor/PlacementMode/Private/SPlacementModeTools.cpp
+++ b/Engine/Source/Editor/PlacementMode/Private/SPlacementModeTools.cpp
@@ -23,6 +23,7 @@
 #include "ThumbnailRendering/ThumbnailManager.h"
 #include "AssetSelection.h"
 #include "ActorFactories/ActorFactory.h"
+#include "ScopedTransaction.h"
 
 #define LOCTEXT_NAMESPACE "PlacementMode"
 
@@ -476,7 +477,6 @@ FReply SPlacementAssetMenuEntry::OnMouseButtonUp(const FGeometry& MyGeometry, co
 	{
 		bIsPressed = false;
 
-		AActor* NewActor = nullptr;
 		UActorFactory* Factory = Item->Factory;
 		if (!Item->Factory)
 		{
@@ -488,13 +488,19 @@ FReply SPlacementAssetMenuEntry::OnMouseButtonUp(const FGeometry& MyGeometry, co
 				FActorFactoryAssetProxy::GetFactoryForAssetObject(ClassObject);
 			}
 		}
-		NewActor = FLevelEditorActionCallbacks::AddActor(Factory, Item->AssetData, nullptr);
-		if (NewActor && GCurrentLevelEditingViewportClient)
+
 		{
-  			GEditor->MoveActorInFrontOfCamera(*NewActor, 
-  				GCurrentLevelEditingViewportClient->GetViewLocation(), 
-  				GCurrentLevelEditingViewportClient->GetViewRotation().Vector()
-  			);
+			// Note: Capture the add and the move within a single transaction, so that the placed actor position is calculated correctly by the transaction diff
+			FScopedTransaction Transaction(NSLOCTEXT("UnrealEd", "CreateActor", "Create Actor"));
+
+			AActor* NewActor = FLevelEditorActionCallbacks::AddActor(Factory, Item->AssetData, nullptr);
+			if (NewActor && GCurrentLevelEditingViewportClient)
+			{
+				GEditor->MoveActorInFrontOfCamera(*NewActor,
+					GCurrentLevelEditingViewportClient->GetViewLocation(),
+					GCurrentLevelEditingViewportClient->GetViewRotation().Vector()
+				);
+			}
 		}
 
 		if (!MouseEvent.IsControlDown())
