Update the placement module files to show the new place actors view.
Changes include:

Enabling drag and drop favorites support
enabling context menu favorites support

switching the asset thumbnail to an image, as we are supporting multiple types and just want a single image

diff --git a/Engine/Source/Editor/PlacementMode/PlacementMode.Build.cs b/Engine/Source/Editor/PlacementMode/PlacementMode.Build.cs
--- a/Engine/Source/Editor/PlacementMode/PlacementMode.Build.cs
+++ b/Engine/Source/Editor/PlacementMode/PlacementMode.Build.cs
@@ -23,6 +23,7 @@ public class PlacementMode : ModuleRules
                 "CollectionManager",
                 "LevelEditor",
                 "AssetTools",
+                "EditorWidgets",
                 "ToolMenus",
                 "WidgetRegistration"
             } 
diff --git a/Engine/Source/Editor/PlacementMode/Private/PlacementModeCommands.cpp b/Engine/Source/Editor/PlacementMode/Private/PlacementModeCommands.cpp
deleted file mode 100644
diff --git a/Engine/Source/Editor/PlacementMode/Private/PlacementModeCommands.h b/Engine/Source/Editor/PlacementMode/Private/PlacementModeCommands.h
deleted file mode 100644
diff --git a/Engine/Source/Editor/PlacementMode/Private/PlacementModeModule.cpp b/Engine/Source/Editor/PlacementMode/Private/PlacementModeModule.cpp
--- a/Engine/Source/Editor/PlacementMode/Private/PlacementModeModule.cpp
+++ b/Engine/Source/Editor/PlacementMode/Private/PlacementModeModule.cpp
@@ -119,11 +119,21 @@ void FPlacementModeModule::StartupModule()
 
 	RegisterPlacementCategory(
 		FPlacementCategoryInfo(
-			NSLOCTEXT("PlacementMode", "RecentlyPlaced", "Recently Placed"),
+			NSLOCTEXT("PlacementMode", "Favorites", "Favorites"),
+			FSlateIcon(),
+			FBuiltInPlacementCategories::Favorites(),
+			TEXT("Favorites"),
+			TNumericLimits<int32>::Lowest(),
+			false
+		));
+	
+	RegisterPlacementCategory(
+		FPlacementCategoryInfo(
+			NSLOCTEXT("PlacementMode", "RecentlyPlaced", "Recent"),
 			FSlateIcon(FAppStyle::GetAppStyleSetName(), "PlacementBrowser.Icons.Recent"),
 			FBuiltInPlacementCategories::RecentlyPlaced(),
 			TEXT("PMRecentlyPlaced"),
-			TNumericLimits<int32>::Lowest(),
+			TNumericLimits<int32>::Lowest() + 1,
 			false
 		)
 	);
@@ -641,6 +651,24 @@ void FPlacementModeModule::GetFilteredItemsForCategory(FName CategoryName, TArra
 	}
 }
 
+void FPlacementModeModule::GetItemsWithNamesForCategory( FName CategoryName, TArray<TSharedPtr<FPlaceableItem>>& OutItems, const TArray<FName>& ItemNames ) const
+{
+	const FPlacementCategory* Category = Categories.Find(CategoryName);
+	if (Category)
+	{
+		for (const TTuple<FGuid, TSharedPtr<FPlaceableItem>>& Pair : Category->Items)
+		{
+			TSharedPtr<FPlaceableItem> Item = Pair.Value;
+			FName Name = FName( Item->NativeName );
+
+			if ( ItemNames.Contains( Name ) )
+			{
+				OutItems.Add(Pair.Value);
+			}
+		}
+	}
+}
+
 void FPlacementModeModule::RegenerateItemsForCategory(FName Category)
 {
 	if (Category == FBuiltInPlacementCategories::RecentlyPlaced())
@@ -651,7 +679,7 @@ void FPlacementModeModule::RegenerateItemsForCategory(FName Category)
 	{
 		RefreshVolumes();
 	}
-	else if (Category == FBuiltInPlacementCategories::AllClasses())
+	else if (Category == FBuiltInPlacementCategories::AllClasses() || Category == FBuiltInPlacementCategories::Favorites())
 	{
 		RefreshAllPlaceableClasses();
 	}
diff --git a/Engine/Source/Editor/PlacementMode/Private/SPlacementModeTools.cpp b/Engine/Source/Editor/PlacementMode/Private/SPlacementModeTools.cpp
--- a/Engine/Source/Editor/PlacementMode/Private/SPlacementModeTools.cpp
+++ b/Engine/Source/Editor/PlacementMode/Private/SPlacementModeTools.cpp
@@ -22,11 +22,14 @@
 #include "Widgets/Docking/SDockTab.h"
 #include "ThumbnailRendering/ThumbnailManager.h"
 #include "AssetSelection.h"
+#include "SAssetDropTarget.h"
 #include "ActorFactories/ActorFactory.h"
 #include "ScopedTransaction.h"
-#include "PlacementModeCommands.h"
 #include "Layout/CategoryDrivenContentBuilder.h"
-#include "Widgets/Layout/SScrollBox.h"
+#include "ToolkitBuilder.h"
+#include "Styles/SlateBrushTemplates.h"
+#include "Layout/WidgetPath.h"
+#include "Styling/CoreStyle.h"
 
 #define LOCTEXT_NAMESPACE "PlacementMode"
 
@@ -123,25 +126,43 @@ public:
 		TSharedPtr<FAssetThumbnailPool> ThumbnailPool = UThumbnailManager::Get().GetSharedThumbnailPool();
 
 		Thumbnail = MakeShareable(new FAssetThumbnail(Asset, InArgs._Width, InArgs._Height, ThumbnailPool));
+		
+		if ( CVarEnableCategoryContentChooserView->GetBool() )
+		{
+			TSharedPtr<SImage> ThumbnailImage;
 
-		FAssetThumbnailConfig Config;
-		Config.bForceGenericThumbnail = InArgs._AlwaysUseGenericThumbnail;
-		Config.ClassThumbnailBrushOverride = InArgs._ClassThumbnailBrushOverride;
-		Config.AssetTypeColorOverride = InArgs._AssetTypeColorOverride;
-		ChildSlot
-		[
-			Thumbnail->MakeThumbnailWidget( Config )
-		];
+			// figure out the proper image to show based on whether the asset is a class type
+			TWeakObjectPtr<UClass> ThumbnailClass = MakeWeakObjectPtr( const_cast<UClass*>( FClassIconFinder::GetIconClassForAssetData( Asset, &bIsClassType ) ) );
+			const FName AssetClassName = Asset.AssetClassPath.GetAssetName();
+			const FName DefaultThumbnail = bIsClassType ? NAME_None : FName( *FString::Printf( TEXT("ClassThumbnail.%s"), *AssetClassName.ToString() ) );
+			const FSlateBrush* ThumbnailBrush = !InArgs._ClassThumbnailBrushOverride.IsNone() ?
+				FClassIconFinder::FindThumbnailForClass( nullptr,  InArgs._ClassThumbnailBrushOverride ) :
+				FClassIconFinder::FindThumbnailForClass( ThumbnailClass.Get(), DefaultThumbnail );
+
+			ChildSlot[ SAssignNew( ThumbnailImage, SImage ).Image( ThumbnailBrush ) ];
+		}
+		else
+		{
+			FAssetThumbnailConfig Config;
+			Config.bForceGenericThumbnail = InArgs._AlwaysUseGenericThumbnail;
+			Config.ClassThumbnailBrushOverride = InArgs._ClassThumbnailBrushOverride;
+			Config.AssetTypeColorOverride = InArgs._AssetTypeColorOverride;
+			ChildSlot[ Thumbnail->MakeThumbnailWidget( Config ) ];			
+		}
 	}
 
 private:
 
 	FAssetData Asset;
 	TSharedPtr< FAssetThumbnail > Thumbnail;
+
+	/** Indicates whether the Asset is a class type */
+	bool bIsClassType;
 };
 
 void SPlacementAssetEntry::Construct(const FArguments& InArgs, const TSharedPtr<const FPlaceableItem>& InItem)
 {	
+	OnGetMenuContent = InArgs._OnGetMenuContent;
 	bIsPressed = false;
 
 	Item = InItem;
@@ -190,18 +211,44 @@ void SPlacementAssetEntry::Construct(const FArguments& InArgs, const TSharedPtr<
 
 	NormalImage = &ButtonStyle.Normal;
 	HoverImage = &ButtonStyle.Hovered;
-	PressedImage = &ButtonStyle.Pressed; 
-
-	// Create doc link widget if there is a class to link to
-	TSharedRef<SWidget> DocWidget = SNew(SSpacer);
-	if(DocClass != NULL)
+	PressedImage = &ButtonStyle.Pressed;
+	const bool bIsClassicView = !CVarEnableCategoryContentChooserView->GetBool();
+	float ThumbnailBoxWidth = 40;
+	
+	float TextFillWidth = 0.99;
+	const FMargin DragHandlePadding{ 0.f,0.f,8.f, 0.f };
+	
+	FMargin WholeAssetPadding{ 8.f, 2.f, 12.f, 2.f };
+	const FSlateBrush* WholeAssetBackgroundBrush{ FAppStyle::Get().GetBrush( "PlacementBrowser.Asset.Background" ) };
+	FMargin ThumbnailBoxPadding{ 8.f ,4.f,8.f, 4.f };
+	FMargin AssetTextPadding{ 9, 0, 0, 1 };
+	TSharedRef<SWidget> DraggableAssetEndWidget = SNullWidget::NullWidget;
+	
+	if ( bIsClassicView )
 	{
-		DocWidget = FEditorClassUtils::GetDocumentationLinkWidget(DocClass);
-		DocWidget->SetCursor( EMouseCursor::Default );
+		// Create doc link widget if there is a class to link to
+		TSharedRef<SWidget> DocWidget = SNew( SSpacer );
+		if( DocClass != NULL )
+		{
+			DocWidget = FEditorClassUtils::GetDocumentationLinkWidget( DocClass );
+			DocWidget->SetCursor( EMouseCursor::Default );
+		}
+		DraggableAssetEndWidget = DocWidget;
 	}
-
+	else
+	{
+		WholeAssetPadding = 0;
+		WholeAssetBackgroundBrush = FAppStyle::Get().GetBrush("PlacementBrowser.Asset.ThumbnailBackground");
+		ThumbnailBoxPadding = FMargin{ 4.f,4.f,0.f, 4.f };
+		AssetTextPadding = FMargin{ 4, 0, 8, 1 };
+		DraggableAssetEndWidget = SNew(SBox )
+									.Padding( DragHandlePadding )
+									[  SNew(SImage).Image( FSlateBrushTemplates::DragHandle() ) ];
+		ThumbnailBoxWidth = 20;
+	}
+	
 	ChildSlot
-	.Padding(FMargin(8.f, 2.f, 12.f, 2.f))
+	.Padding( WholeAssetPadding )
 	[
 
 		SNew(SOverlay)
@@ -209,7 +256,7 @@ void SPlacementAssetEntry::Construct(const FArguments& InArgs, const TSharedPtr<
 		+SOverlay::Slot()
 		[
 			SNew(SBorder)
-			.BorderImage( FAppStyle::Get().GetBrush("PlacementBrowser.Asset.Background"))
+			.BorderImage( WholeAssetBackgroundBrush)
 			.Cursor( EMouseCursor::GrabHand )
 			.ToolTip( AssetEntryToolTip )
 			.Padding(0)
@@ -218,13 +265,13 @@ void SPlacementAssetEntry::Construct(const FArguments& InArgs, const TSharedPtr<
 				SNew( SHorizontalBox )
 
 				+ SHorizontalBox::Slot()
-				.Padding(8.0f, 4.f)
+				.Padding( ThumbnailBoxPadding )
 				.AutoWidth()
 				.HAlign(HAlign_Center)
 				.VAlign(VAlign_Center)
 				[
 					SNew( SBox )
-					.WidthOverride(40)
+					.WidthOverride( ThumbnailBoxWidth )
 					.HeightOverride(40)
 					[
 						SNew( SPlacementAssetThumbnail, Item->AssetData )
@@ -243,22 +290,21 @@ void SPlacementAssetEntry::Construct(const FArguments& InArgs, const TSharedPtr<
 					.BorderImage(FAppStyle::Get().GetBrush("PlacementBrowser.Asset.LabelBack"))
 					[
 						SNew( SHorizontalBox)
-						+SHorizontalBox::Slot()
-						.Padding(9, 0, 0, 1)
+						+SHorizontalBox::Slot()	
+						.FillContentWidth( TextFillWidth )
+						.Padding( AssetTextPadding )
 						.VAlign(VAlign_Center)
 						[
 							SNew( STextBlock )
 							.TextStyle( FAppStyle::Get(), "PlacementBrowser.Asset.Name" )
 							.Text( Item->DisplayName )
-							.HighlightText(InArgs._HighlightText)
+							.OverflowPolicy( ETextOverflowPolicy::Ellipsis )
+							.HighlightText(  InArgs._HighlightText )
 						]
-
 						+ SHorizontalBox::Slot()
-						.VAlign(VAlign_Center)
-						.AutoWidth()
-						[
-							DocWidget
-						]
+								.VAlign( VAlign_Center )
+								.AutoWidth()
+								[ DraggableAssetEndWidget ]
 					]
 				]
 			]
@@ -283,6 +329,22 @@ FReply SPlacementAssetEntry::OnMouseButtonDown(const FGeometry& MyGeometry, cons
 		return FReply::Handled().DetectDrag( SharedThis( this ), MouseEvent.GetEffectingButton() );
 	}
 
+	// Create the context menu to be launched on right mouse click. 
+	if ( CVarEnableCategoryContentChooserView->GetBool() && MouseEvent.GetEffectingButton() == EKeys::RightMouseButton )
+	{
+		FWidgetPath WidgetPath = MouseEvent.GetEventPath() != nullptr ? *MouseEvent.GetEventPath() : FWidgetPath();
+
+		FSlateApplication::Get().PushMenu(
+			AsShared(),
+			WidgetPath,
+			OnGetMenuContent.IsBound() ? OnGetMenuContent.Execute() : SNullWidget::NullWidget,
+			MouseEvent.GetScreenSpacePosition(),
+			FPopupTransitionEffect( FPopupTransitionEffect::ContextMenu )
+			);
+
+		return FReply::Handled();
+	}
+
 	return FReply::Unhandled();
 }
 
@@ -600,11 +662,29 @@ void SPlacementModeTools::Construct( const FArguments& InArgs, TSharedRef<SDockT
 	bRefreshRecentlyPlaced = false;
 	bUpdateShownItems = true;
 
-	FCategoryDrivenContentBuilderArgs Args("PlacementModes");
+	FCategoryDrivenContentBuilderArgs Args( "PlacementModes", UE::DisplayBuilders::FBuilderKeys::Get().PlaceActors() );
+	Args.FavoritesCommandName = FBuiltInPlacementCategories::Favorites();
+	Args.ActiveCategoryName = FBuiltInPlacementCategories::Basic();
+	Args.GetDecoratedButtonDelegate.BindLambda( [ this ] ( TSharedRef<SWidget> UndecoratedButtonWidget )
+	{
+		return SNew( SAssetDropTarget )
+			.bOnlyRecognizeOnDragEnter( false )
+			.OnAssetsDropped_Lambda( [ this ] ( const FDragDropEvent& Event, TArrayView<FAssetData> InAssets )
+			{
+				if ( InAssets.Num() == 1 )
+				{
+					CategoryContentBuilder->AddFavorite( InAssets[0].AssetName );
+				}
+			})
+			[
+				UndecoratedButtonWidget	
+			];
+	});
+
+	
+
 	CategoryContentBuilder = MakeShared<FCategoryDrivenContentBuilder>( Args );
-	CategoryContentBuilder->ProvideSelectedCategoryContentDelegate.BindSP(this, &SPlacementModeTools::ProvideCategoryContent);
-	UniformWrapPanel = SNew( SUniformWrapPanel );
-	SAssignNew(CustomContent, SBox);
+	CategoryContentBuilder->UpdateContentForCategoryDelegate.BindSP( SharedThis( this ), &SPlacementModeTools::UpdateContentForCategory );
 
 	ActiveTabName = FBuiltInPlacementCategories::Basic();
 
@@ -644,7 +724,7 @@ void SPlacementModeTools::Construct( const FArguments& InArgs, TSharedRef<SDockT
 			.FillHeight(1)
 			[
 				SNew(SBorder)
-				.BorderImage(FAppStyle::GetBrush("ToolPanel.GroupBorder"))
+				.BorderImage( FSlateBrushTemplates::Panel() )
 				.Visibility_Lambda( [] () { return CVarEnableCategoryContentChooserView->GetBool()   ? EVisibility::Visible : EVisibility::Collapsed; })
 				.Padding(0)
 				[
@@ -762,7 +842,7 @@ void SPlacementModeTools::UpdateShownItems()
 	}
 	else if (Category->CustomGenerator)
 	{
-		UniformWrapPanel->AddSlot()[Category->CustomGenerator()];
+		CategoryContentBuilder->FillWithBuilder( Category->CustomGenerator() );
 		CustomContent->SetContent(Category->CustomGenerator());
 
 		CustomContent->SetVisibility(EVisibility::Visible);
@@ -770,6 +850,11 @@ void SPlacementModeTools::UpdateShownItems()
 
 		FilterLabelPtr->SetText(Category->DisplayName);
 	}
+	else if ( IsFavoritesCategorySelected() )
+	{
+		IPlacementModeModule::Get().RegenerateItemsForCategory( FBuiltInPlacementCategories::AllClasses() );
+		PlacementModeModule.GetItemsWithNamesForCategory( FBuiltInPlacementCategories::AllClasses(), FavoriteItems, CategoryContentBuilder->GetFavorites() );
+	}
 	else
 	{
 		FilteredItems.Reset();
@@ -839,39 +924,62 @@ TSharedRef<ITableRow> SPlacementModeTools::OnGenerateWidgetForItem(TSharedPtr<FP
 
 TSharedRef<SWidget> SPlacementModeTools::GetPlacementAssetWidget( const TSharedPtr<FPlaceableItem>& InItem ) const
 {
-	return SNew(SPlacementAssetEntry, InItem.ToSharedRef())
-		.HighlightText(this, &SPlacementModeTools::GetHighlightText);
+	TSharedRef<SPlacementAssetEntry> Entry =  SNew( SPlacementAssetEntry, InItem.ToSharedRef() )
+		.HighlightText(this, &SPlacementModeTools::GetHighlightText)
+		.Clipping( EWidgetClipping::ClipToBounds )
+		.OnGetMenuContent_Lambda( [this, &InItem] ()
+		{
+			return CategoryContentBuilder.IsValid() ? CategoryContentBuilder->CreateFavoritesContextMenu( InItem->NativeName ) : SNullWidget::NullWidget;
+		});
+	return Entry;
 }
 
-TSharedRef<SWidget> SPlacementModeTools::ProvideCategoryContent(const FName& CategoryName)
+void SPlacementModeTools::UpdateContentForCategory( FName CategoryName, FText CategoryLabel )
 {
-	if (CVarEnableCategoryContentChooserView->GetBool())
+	if ( CVarEnableCategoryContentChooserView->GetBool() )
 	{
-		SetActiveTab(CategoryName);
-		UniformWrapPanel->ClearChildren();
+		SetActiveTab( CategoryName );
+		CustomContent = SNew( SBox );
+		FavoriteItems.Empty();
+
+		CategoryContentBuilder->ClearCategoryContent();
+		
+		// if the Category name is not none, the user updated the category, so clear out the search ~ the Category choice should override it.
+		// The call of UpdateShownItems below will update search state based on this setting.
+		if ( !CategoryName.IsNone() )
+		{
+			SearchBoxPtr->SetText( FText::GetEmpty());
+		}
+		
 		UpdateShownItems();
 
-		// else, add the items found for the category (and potentially filtered)to our display boxes to show 
-		if ( !FilteredItems.IsEmpty() && UniformWrapPanel->GetChildren()->Num() == 0 )
+		const FPlacementCategoryInfo* Category = IPlacementModeModule::Get().GetRegisteredPlacementCategory( CategoryName );
+		if ( Category && Category->CustomGenerator )
+		{
+			CategoryContentBuilder->FillWithBuilder( Category->CustomGenerator() );
+			return;
+		}
+
+		 if ( IsFavoritesCategorySelected() )
+		{
+				for ( const TSharedPtr<FPlaceableItem>& Item : FavoriteItems )
+				{
+					CategoryContentBuilder->AddBuilder( GetPlacementAssetWidget( Item ) );
+				}				
+		}
+		else
 		{
 			for (const TSharedPtr<FPlaceableItem>& Item : FilteredItems)
 			{
-				UniformWrapPanel->AddSlot()
-				[
-					GetPlacementAssetWidget(Item)
-				];
+				CategoryContentBuilder->AddBuilder( GetPlacementAssetWidget(Item) );
 			}
 		}
-
-		return UniformWrapPanel->GetChildren()->Num() == 0 ?
-			SNullWidget::NullWidget :
-			SNew( SScrollBox )
-						+ SScrollBox::Slot()
-						[
-							UniformWrapPanel.ToSharedRef()
-						];
 	}
-	return SNullWidget::NullWidget;
+}
+
+bool SPlacementModeTools::IsFavoritesCategorySelected() const
+{
+	return ActiveTabName == FBuiltInPlacementCategories::Favorites() && !IsSearchActive();
 }
 
 void SPlacementModeTools::OnCategoryChanged(const ECheckBoxState NewState, FName InCategory)
@@ -925,14 +1033,19 @@ void SPlacementModeTools::UpdatePlacementCategories()
 
 	TArray<FPlacementCategoryInfo> Categories;
 	IPlacementModeModule::Get().GetSortedCategories(Categories);
-	
-	// register will make a command to load each Category
-	FPlacementModeCommands::Register();
-	// set the commands as the Categories in our category content chooser
-	CategoryContentBuilder->SetCommands(FPlacementModeCommands::Get().GetPlacementToolkitCategoryCommands());
-	
+
+	TArray<UE::DisplayBuilders::FBuilderInput> BuilderInputArray;
+
 	for (const FPlacementCategoryInfo& Category : Categories)
 	{
+		const UE::DisplayBuilders::FBuilderInput& InputInfo = UE::DisplayBuilders::FBuilderInput( Category.UniqueHandle, Category.DisplayName,
+		Category.DisplayIcon, EUserInterfaceActionType::ToggleButton );
+		BuilderInputArray.Add( InputInfo );
+		
+		if ( Category.UniqueHandle == FBuiltInPlacementCategories::Favorites() && !CVarEnableCategoryContentChooserView->GetBool() )
+		{
+			continue;
+		}
 		if (Category.UniqueHandle == FBuiltInPlacementCategories::Basic())
 		{
 			BasicTabExists = true;
@@ -958,7 +1071,8 @@ void SPlacementModeTools::UpdatePlacementCategories()
 			]
 		];
 	}
-
+	CategoryContentBuilder->InitializeCategoryButtons( BuilderInputArray );
+	
 	if (TabToActivate.IsNone())
 	{
 		if (BasicTabExists)
@@ -1007,11 +1121,15 @@ void SPlacementModeTools::OnSearchChanged(const FText& InFilterText)
 		bUpdateShownItems = true;
 	}
 
+	const FText OldText = SearchTextFilter->GetRawFilterText();
 	SearchTextFilter->SetRawFilterText( InFilterText );
 	SearchBoxPtr->SetError( SearchTextFilter->GetFilterErrorText() );
 
-	// update the category content which may have changed on search
-	CategoryContentBuilder->UpdateWidget();
+	if ( !OldText.EqualToCaseIgnored( InFilterText ) )
+	{
+		CategoryContentBuilder->SetShowNoCategorySelection( IsSearchActive() );
+		CategoryContentBuilder->UpdateWidget();
+	}
 }
 
 void SPlacementModeTools::OnSearchCommitted(const FText& InFilterText, ETextCommit::Type InCommitType)
diff --git a/Engine/Source/Editor/PlacementMode/Private/SPlacementModeTools.h b/Engine/Source/Editor/PlacementMode/Private/SPlacementModeTools.h
--- a/Engine/Source/Editor/PlacementMode/Private/SPlacementModeTools.h
+++ b/Engine/Source/Editor/PlacementMode/Private/SPlacementModeTools.h
@@ -28,6 +28,7 @@ public:
 
 		/** Highlight this text in the text block */
 		SLATE_ATTRIBUTE(FText, HighlightText)
+		SLATE_EVENT( FOnGetContent, OnGetMenuContent )
 
 	SLATE_END_ARGS()
 
@@ -39,9 +40,14 @@ public:
 
 	bool IsPressed() const;
 
+
 	TSharedPtr<const FPlaceableItem> Item;
 
 private:
+	
+	/** Delegate to execute to get the menu content of this button */
+	FOnGetContent OnGetMenuContent;
+	
 	const FSlateBrush* GetBorder() const;
 
 	bool bIsPressed;
@@ -102,13 +108,12 @@ public:
 	TSharedRef<SWidget> GetPlacementAssetWidget(const TSharedPtr<FPlaceableItem>& InItem) const;
 
 	/**
-	 * Given FName CategoryName as the Category that has been clicked, it should return the
-	 * TSharedRef<SWidget> that contains the content for that Category
-	 * 
-	 * @param CategoryName the name of the Category that the user has clicked
-	 * @return 
+	 * Given FName CategoryName as the Category that has been clicked, update the FCategoryDrivenContentBuilder with the proper content for the category
+	 *
+	 * @param CategoryName the name of the currently chosen Category
+	 * @param CategoryLabel the label for the currently chosen Category 
 	 */
-	TSharedRef<SWidget> ProvideCategoryContent(const FName& CategoryName);
+	void UpdateContentForCategory( FName CategoryName, FText CategoryLabel );
 
 private:
 
@@ -130,6 +135,11 @@ private:
 	/** Check if a search is active */
 	bool IsSearchActive() const;
 
+	/**
+	 * returns true is the favorites category is selected, else it returns false
+	 */
+	bool IsFavoritesCategorySelected() const;
+	
 	/** Update the list of shown items */
 	void UpdateShownItems();
 
@@ -211,9 +221,9 @@ private:
 	/** List view that shows placeable items */
 	TSharedPtr<SListView<TSharedPtr<FPlaceableItem>>> ListView;
 
-	/** Wrap panel containing the draggable content  */
-	TSharedPtr<SUniformWrapPanel> UniformWrapPanel;
-	
-	/** The FCategoryDrivenContentBuilder  which will build the UI for this mode, if defined */
+	/** The FCategoryDrivenContentBuilder which will build the UI for this mode, if defined */
 	TSharedPtr<FCategoryDrivenContentBuilder> CategoryContentBuilder;
+
+	/** Array of favorite items to show in the favorites view */
+	TArray<TSharedPtr<FPlaceableItem>> FavoriteItems;
 };
diff --git a/Engine/Source/Editor/PlacementMode/Public/IPlacementModeModule.h b/Engine/Source/Editor/PlacementMode/Public/IPlacementModeModule.h
--- a/Engine/Source/Editor/PlacementMode/Public/IPlacementModeModule.h
+++ b/Engine/Source/Editor/PlacementMode/Public/IPlacementModeModule.h
@@ -16,6 +16,7 @@
 #include "Textures/SlateIcon.h"
 
 class FNamePermissionList;
+class FBuilderInput;
 
 /**
  * Struct that defines an identifier for a particular placeable item in this module.
@@ -250,6 +251,7 @@ private:
 /** Structure of built-in placement categories. Defined as functions to enable external use without linkage */
 struct FBuiltInPlacementCategories
 {
+	static FName Favorites()		{ static FName Name("Favorites");		return Name; }
 	static FName RecentlyPlaced()	{ static FName Name("RecentlyPlaced");	return Name; }
 	static FName Basic()			{ static FName Name("Basic");			return Name; }
 	static FName Lights()			{ static FName Name("Lights");			return Name; }
@@ -431,6 +433,15 @@ public:
 	 */
 	virtual void GetFilteredItemsForCategory(FName Category, TArray<TSharedPtr<FPlaceableItem>>& OutItems, TFunctionRef<bool(const TSharedPtr<FPlaceableItem>&)> Filter) const = 0;
 
+	/**
+	 * Puts the items that should be shown for Category CategoryName in the OutItems array
+	 * 
+	 * @param CategoryName the name of the category to get the items for
+	 * @param OutItems the array of FPlaceableItems to put the items in
+	 * @param ItemNames the names of the FPlaceableItems to get and put in the OutItems array
+	 */
+	virtual void GetItemsWithNamesForCategory( FName CategoryName, TArray<TSharedPtr<FPlaceableItem>>& OutItems, const TArray<FName>& ItemNames ) const = 0;
+
 	/**
 	 * Instruct the category associated with the specified unique handle that it should regenerate its items
 	 *
diff --git a/Engine/Source/Editor/PlacementMode/Public/PlacementModeModule.h b/Engine/Source/Editor/PlacementMode/Public/PlacementModeModule.h
--- a/Engine/Source/Editor/PlacementMode/Public/PlacementModeModule.h
+++ b/Engine/Source/Editor/PlacementMode/Public/PlacementModeModule.h
@@ -11,7 +11,7 @@ struct FPlacementCategory : FPlacementCategoryInfo
 	FPlacementCategory(const FPlacementCategoryInfo& SourceInfo)
 		: FPlacementCategoryInfo(SourceInfo)
 	{
-
+		
 	}
 
 	FPlacementCategory(FPlacementCategory&& In)
@@ -109,6 +109,15 @@ public:
 
 	virtual void GetFilteredItemsForCategory(FName CategoryName, TArray<TSharedPtr<FPlaceableItem>>& OutItems, TFunctionRef<bool(const TSharedPtr<FPlaceableItem>&)> Filter) const;
 
+	/**
+	 * Puts the FPlaceableItem instances with names matching the FNames in ItemNames for the placement category CategoryName in the OutItems TArray
+	 * 
+	 * @param CategoryName the name of the Placement category that contains the FPlaceableItems
+	 * @param OutItems  the TArray to fill with the Items that match the names in ItemNames
+	 * @param ItemNames the names of the FPlaceableItem instances to put in OutItems
+	 */
+	virtual void GetItemsWithNamesForCategory( FName CategoryName, TArray<TSharedPtr<FPlaceableItem>>& OutItems, const TArray<FName>& ItemNames  ) const override;
+
 	virtual void RegenerateItemsForCategory(FName Category) override;
 
 private:
