Place Actors
- Renamed "Visual Effects" to "VFX", "All Classes" to "All" and "Cinematic" to "Cine" so that they fit the new buttons.
- Removed drag and dropping into favorites support since it's incomplete.
- Fixed a bug where the CategoryDrivenContentBuilder was creating a new widget without going through its GenerateWidget path. This meant that subsequent calls to UpdateWidget would update the new widget but the UI would still display the old one.
- Changed the sorting for "All Classes" so that it doesn't take the order override into account. The order override doesn't make sense in a global scope and had all of the brushes as the first items for All Classes.

diff --git a/Engine/Plugins/Experimental/CineCameraRigs/Source/CineCameraRigsEditor/Private/CineCameraRigsEditor.cpp b/Engine/Plugins/Experimental/CineCameraRigs/Source/CineCameraRigsEditor/Private/CineCameraRigsEditor.cpp
--- a/Engine/Plugins/Experimental/CineCameraRigs/Source/CineCameraRigsEditor/Private/CineCameraRigsEditor.cpp
+++ b/Engine/Plugins/Experimental/CineCameraRigs/Source/CineCameraRigsEditor/Private/CineCameraRigsEditor.cpp
@@ -76,7 +76,7 @@ private:
 		}
 
 		FPlacementCategoryInfo Info(
-			LOCTEXT("CinematicCategoryName", "Cinematic"),
+			LOCTEXT("CinematicCategoryName", "Cine"),
 			FSlateIcon(FAppStyle::GetAppStyleSetName(), "PlacementBrowser.Icons.Cinematics"),
 			"Cinematic",
 			TEXT("PMCinematic"),
diff --git a/Engine/Plugins/MovieScene/LevelSequenceEditor/Source/LevelSequenceEditor/Private/LevelSequenceEditorModule.cpp b/Engine/Plugins/MovieScene/LevelSequenceEditor/Source/LevelSequenceEditor/Private/LevelSequenceEditorModule.cpp
--- a/Engine/Plugins/MovieScene/LevelSequenceEditor/Source/LevelSequenceEditor/Private/LevelSequenceEditorModule.cpp
+++ b/Engine/Plugins/MovieScene/LevelSequenceEditor/Source/LevelSequenceEditor/Private/LevelSequenceEditorModule.cpp
@@ -153,7 +153,7 @@ protected:
 	void RegisterPlacementModeExtensions()
 	{
 		FPlacementCategoryInfo Info(
-			LOCTEXT("CinematicCategoryName", "Cinematic"),
+			LOCTEXT("CinematicCategoryName", "Cine"),
 			FSlateIcon(FAppStyle::GetAppStyleSetName(), "PlacementBrowser.Icons.Cinematics"),
 			"Cinematic",
 			TEXT("PMCinematic"),
diff --git a/Engine/Source/Developer/WidgetRegistration/Private/Layout/CategoryDrivenContentBuilder.cpp b/Engine/Source/Developer/WidgetRegistration/Private/Layout/CategoryDrivenContentBuilder.cpp
--- a/Engine/Source/Developer/WidgetRegistration/Private/Layout/CategoryDrivenContentBuilder.cpp
+++ b/Engine/Source/Developer/WidgetRegistration/Private/Layout/CategoryDrivenContentBuilder.cpp
@@ -309,7 +309,7 @@ void FCategoryDrivenContentBuilder::InitializeCategoryButtons()
 	const FName StyleName = GetCategoryToolBarStyleName(); 
 	LoadPaletteToolBarBuilder->SetStyle(&FAppStyle::Get(), StyleName );
 	LoadPaletteToolBarBuilder->SetLabelVisibility( CategoryButtonLabelVisibility );
-	MainContentVerticalBox = SNew(SVerticalBox);
+
 	InitializeCategoryToolbar();
 }
 
diff --git a/Engine/Source/Editor/PlacementMode/Private/PlacementModeModule.cpp b/Engine/Source/Editor/PlacementMode/Private/PlacementModeModule.cpp
--- a/Engine/Source/Editor/PlacementMode/Private/PlacementModeModule.cpp
+++ b/Engine/Source/Editor/PlacementMode/Private/PlacementModeModule.cpp
@@ -120,7 +120,7 @@ void FPlacementModeModule::StartupModule()
 	RegisterPlacementCategory(
 		FPlacementCategoryInfo(
 			NSLOCTEXT("PlacementMode", "Favorites", "Favorites"),
-			FSlateIcon(),
+			FSlateIcon(FAppStyle::Get().GetStyleSetName(), "Icons.Favorites.Small"),
 			FBuiltInPlacementCategories::Favorites(),
 			TEXT("Favorites"),
 			TNumericLimits<int32>::Lowest(),
@@ -241,7 +241,7 @@ void FPlacementModeModule::StartupModule()
 		FName CategoryName = FBuiltInPlacementCategories::Visual();
 		RegisterPlacementCategory(
 			FPlacementCategoryInfo(
-				NSLOCTEXT("PlacementMode", "VisualEffects", "Visual Effects"),
+				NSLOCTEXT("PlacementMode", "VisualEffects", "VFX"),
 				FSlateIcon(FAppStyle::GetAppStyleSetName(), "PlacementBrowser.Icons.VisualEffects"),
 				CategoryName,
 				TEXT("PMVisual"),
@@ -275,7 +275,7 @@ void FPlacementModeModule::StartupModule()
 
 	RegisterPlacementCategory(
 		FPlacementCategoryInfo(
-			NSLOCTEXT("PlacementMode", "AllClasses", "All Classes"),
+			NSLOCTEXT("PlacementMode", "AllClasses", "All"),
 			FSlateIcon(FAppStyle::GetAppStyleSetName(), "PlacementBrowser.Icons.All"),
 			FBuiltInPlacementCategories::AllClasses(),
 			TEXT("PMAllClasses"),
@@ -939,4 +939,4 @@ bool FPlacementModeModule::PassesFilters(const TSharedPtr<FPlaceableItem>& Item)
 void FPlacementModeModule::OnCategoryPermissionListChanged()
 {
 	PlacementModeCategoryListChanged.Broadcast();
-}
\ No newline at end of file
+}
diff --git a/Engine/Source/Editor/PlacementMode/Private/SPlacementModeTools.cpp b/Engine/Source/Editor/PlacementMode/Private/SPlacementModeTools.cpp
--- a/Engine/Source/Editor/PlacementMode/Private/SPlacementModeTools.cpp
+++ b/Engine/Source/Editor/PlacementMode/Private/SPlacementModeTools.cpp
@@ -685,23 +685,6 @@ void SPlacementModeTools::Construct( const FArguments& InArgs, TSharedRef<SDockT
 	FCategoryDrivenContentBuilderArgs Args( "PlacementModes", UE::DisplayBuilders::FBuilderKeys::Get().PlaceActors() );
 	Args.FavoritesCommandName = FBuiltInPlacementCategories::Favorites();
 	Args.ActiveCategoryName = FBuiltInPlacementCategories::Basic();
-	Args.GetDecoratedButtonDelegate.BindLambda( [ this ] ( TSharedRef<SWidget> UndecoratedButtonWidget )
-	{
-		return SNew( SAssetDropTarget )
-			.bOnlyRecognizeOnDragEnter( false )
-			.OnAssetsDropped_Lambda( [ this ] ( const FDragDropEvent& Event, TArrayView<FAssetData> InAssets )
-			{
-				if ( InAssets.Num() == 1 )
-				{
-					CategoryContentBuilder->AddFavorite( InAssets[0].AssetName );
-				}
-			})
-			[
-				UndecoratedButtonWidget	
-			];
-	});
-
-	
 
 	CategoryContentBuilder = MakeShared<FCategoryDrivenContentBuilder>( Args );
 	CategoryContentBuilder->UpdateContentForCategoryDelegate.BindSP( SharedThis( this ), &SPlacementModeTools::UpdateContentForCategory );
@@ -906,7 +889,15 @@ void SPlacementModeTools::UpdateShownItems()
 
 			if (Category->bSortable)
 			{
-				FilteredItems.Sort(&FSortPlaceableItems::SortItemsByOrderThenName);
+				// The item order makes sense internally to a category, not across all classes, so sort by name only in the all classes case
+				if (Category->UniqueHandle == FBuiltInPlacementCategories::AllClasses())
+				{
+					FilteredItems.Sort(&FSortPlaceableItems::SortItemsByName);
+				}
+				else
+				{
+					FilteredItems.Sort(&FSortPlaceableItems::SortItemsByOrderThenName);
+				}
 			}
 		}
 
@@ -1057,7 +1048,7 @@ void SPlacementModeTools::OnCategoryRefresh(FName CategoryName)
 
 void SPlacementModeTools::UpdatePlacementCategories()
 {
-	bool BasicTabExists = false;
+	bool bBasicTabExists = false;
 	FName TabToActivate;
 
 	CategoryFilterPtr->ClearChildren();
@@ -1079,7 +1070,7 @@ void SPlacementModeTools::UpdatePlacementCategories()
 		}
 		if (Category.UniqueHandle == FBuiltInPlacementCategories::Basic())
 		{
-			BasicTabExists = true;
+			bBasicTabExists = true;
 		}
 
 		if (Category.UniqueHandle == ActiveTabName)
@@ -1106,7 +1097,7 @@ void SPlacementModeTools::UpdatePlacementCategories()
 	
 	if (TabToActivate.IsNone())
 	{
-		if (BasicTabExists)
+		if (bBasicTabExists)
 		{
 			TabToActivate = FBuiltInPlacementCategories::Basic();
 		}
