Place Actors Panel
- Support brushes in the "Add+" menu by removing the CustomGenerator using SBspPalette which wasn't needed anymore and adding DragHandler support in FPlacementMenuEntry
- Added support for short category names to fit buttons. Section headers still use the full names.
- Removed mention of dragging from the empty favorites hint text
- Removed unused bInUseLongDisplayName from FBuilderInput

diff --git a/Engine/Plugins/Editor/GeometryMode/Source/BspMode/Private/BspModeModule.cpp b/Engine/Plugins/Editor/GeometryMode/Source/BspMode/Private/BspModeModule.cpp
--- a/Engine/Plugins/Editor/GeometryMode/Source/BspMode/Private/BspModeModule.cpp
+++ b/Engine/Plugins/Editor/GeometryMode/Source/BspMode/Private/BspModeModule.cpp
@@ -10,7 +10,6 @@
 #include "Builders/SpiralStairBuilder.h"
 #include "Builders/TetrahedronBuilder.h"
 #include "BspModeStyle.h"
-#include "SBspPalette.h"
 #include "IPlacementModeModule.h"
 #include "DragAndDrop/BrushBuilderDragDropOp.h"
 #include "BspDragHandler.h"
@@ -45,7 +44,6 @@ void FBspModeModule::StartupModule()
 */
 
 	FPlacementCategoryInfo CategoryInfo (LOCTEXT("PlacementMode_Geometry", "Geometry"), FSlateIcon(FAppStyle::GetAppStyleSetName(), "PlacementBrowser.Icons.BSP"), CategoryName, TEXT("PMGeometry"), 35);
-	CategoryInfo.CustomGenerator = []() -> TSharedRef<SWidget> { return SNew(SBspPalette); };
 	IPlacementModeModule::Get().RegisterPlacementCategory( CategoryInfo );
 	
 	RegisterBspBuilderType(UCubeBuilder::StaticClass(), LOCTEXT("CubeBuilderName", "Box"), LOCTEXT("CubeBuilderToolTip", "Make a box brush"), FBspModeStyle::Get().GetBrush(TEXT("BspMode.BoxBrush")));
@@ -83,7 +81,7 @@ void FBspModeModule::RegisterBspBuilderType( class UClass* InBuilderClass, const
 	static int32 SortOrder = 0;
 	
 	BuilderType->PlaceableItem =  MakeShared<FPlaceableItem>( Handler, SortOrder++, BuilderType->Text, Name );
-	BuilderType->PlacementModeID = IPlacementModeModule::Get().RegisterPlaceableItem( CategoryName,   BuilderType->PlaceableItem.ToSharedRef() );
+	BuilderType->PlacementModeID = IPlacementModeModule::Get().RegisterPlaceableItem( CategoryName, BuilderType->PlaceableItem.ToSharedRef() );
 }
 
 
diff --git a/Engine/Plugins/Experimental/CineCameraRigs/Source/CineCameraRigsEditor/Private/CineCameraRigsEditor.cpp b/Engine/Plugins/Experimental/CineCameraRigs/Source/CineCameraRigsEditor/Private/CineCameraRigsEditor.cpp
--- a/Engine/Plugins/Experimental/CineCameraRigs/Source/CineCameraRigsEditor/Private/CineCameraRigsEditor.cpp
+++ b/Engine/Plugins/Experimental/CineCameraRigs/Source/CineCameraRigsEditor/Private/CineCameraRigsEditor.cpp
@@ -76,12 +76,13 @@ private:
 		}
 
 		FPlacementCategoryInfo Info(
-			LOCTEXT("CinematicCategoryName", "Cine"),
+			LOCTEXT("CinematicCategoryName", "Cinematic"),
 			FSlateIcon(FAppStyle::GetAppStyleSetName(), "PlacementBrowser.Icons.Cinematics"),
 			"Cinematic",
 			TEXT("PMCinematic"),
 			25
 		);
+		Info.ShortDisplayName = LOCTEXT("CinematicShortCategoryName", "Cine");
 
 		PlacmentModeModule.RegisterPlacementCategory(Info);
 		return PlacmentModeModule.GetRegisteredPlacementCategory("Cinematic");
diff --git a/Engine/Plugins/Experimental/VirtualProductionUtilities/Source/VPUtilitiesEditor/Private/VPUtilitiesEditorModule.cpp b/Engine/Plugins/Experimental/VirtualProductionUtilities/Source/VPUtilitiesEditor/Private/VPUtilitiesEditorModule.cpp
--- a/Engine/Plugins/Experimental/VirtualProductionUtilities/Source/VPUtilitiesEditor/Private/VPUtilitiesEditorModule.cpp
+++ b/Engine/Plugins/Experimental/VirtualProductionUtilities/Source/VPUtilitiesEditor/Private/VPUtilitiesEditorModule.cpp
@@ -96,7 +96,7 @@ const FPlacementCategoryInfo* FVPUtilitiesEditorModule::GetVirtualProductionPlac
 				TEXT("PMVirtualProduction"),
 				25
 			);
-
+			Info.ShortDisplayName = LOCTEXT("VirtualProductionShortCategoryName", "VP");
 			IPlacementModeModule::Get().RegisterPlacementCategory(Info);
 
 			// This will return nullptr if the Register above failed so we don't need to explicitly check
diff --git a/Engine/Plugins/Media/MediaPlate/Source/MediaPlateEditor/Private/MediaPlateEditorModule.cpp b/Engine/Plugins/Media/MediaPlate/Source/MediaPlateEditor/Private/MediaPlateEditorModule.cpp
--- a/Engine/Plugins/Media/MediaPlate/Source/MediaPlateEditor/Private/MediaPlateEditorModule.cpp
+++ b/Engine/Plugins/Media/MediaPlate/Source/MediaPlateEditor/Private/MediaPlateEditorModule.cpp
@@ -176,6 +176,7 @@ const FPlacementCategoryInfo* FMediaPlateEditorModule::GetMediaCategoryRegistere
 			26 // Determines where the category shows up in the list with respect to the others.
 		);
 
+		Info.ShortDisplayName = LOCTEXT("MediaPlateShortCategoryName", "Media");
 		IPlacementModeModule::Get().RegisterPlacementCategory(Info);
 
 		return PlacementModeModule.GetRegisteredPlacementCategory(PlacementModeCategoryHandle);
diff --git a/Engine/Plugins/MovieScene/LevelSequenceEditor/Source/LevelSequenceEditor/Private/LevelSequenceEditorModule.cpp b/Engine/Plugins/MovieScene/LevelSequenceEditor/Source/LevelSequenceEditor/Private/LevelSequenceEditorModule.cpp
--- a/Engine/Plugins/MovieScene/LevelSequenceEditor/Source/LevelSequenceEditor/Private/LevelSequenceEditorModule.cpp
+++ b/Engine/Plugins/MovieScene/LevelSequenceEditor/Source/LevelSequenceEditor/Private/LevelSequenceEditorModule.cpp
@@ -153,12 +153,13 @@ protected:
 	void RegisterPlacementModeExtensions()
 	{
 		FPlacementCategoryInfo Info(
-			LOCTEXT("CinematicCategoryName", "Cine"),
+			LOCTEXT("CinematicCategoryName", "Cinematic"),
 			FSlateIcon(FAppStyle::GetAppStyleSetName(), "PlacementBrowser.Icons.Cinematics"),
 			"Cinematic",
 			TEXT("PMCinematic"),
 			25
 		);
+		Info.ShortDisplayName = LOCTEXT("CinematicShortCategoryName", "Cine");
 
 		IPlacementModeModule::Get().RegisterPlacementCategory(Info);
 		IPlacementModeModule::Get().RegisterPlaceableItem(Info.UniqueHandle, MakeShareable( new FPlaceableItem(nullptr, FAssetData(ACameraShakeSourceActor::StaticClass())) ));
diff --git a/Engine/Plugins/VirtualProduction/CameraCalibration/Source/CameraCalibrationEditor/Private/CameraCalibrationEditorModule.cpp b/Engine/Plugins/VirtualProduction/CameraCalibration/Source/CameraCalibrationEditor/Private/CameraCalibrationEditorModule.cpp
--- a/Engine/Plugins/VirtualProduction/CameraCalibration/Source/CameraCalibrationEditor/Private/CameraCalibrationEditorModule.cpp
+++ b/Engine/Plugins/VirtualProduction/CameraCalibration/Source/CameraCalibrationEditor/Private/CameraCalibrationEditorModule.cpp
@@ -73,7 +73,7 @@ const FPlacementCategoryInfo* FCameraCalibrationEditorModule::GetVirtualProducti
 			TEXT("PMVirtualProduction"),
 			25
 		);
-
+		Info.ShortDisplayName = LOCTEXT("VirtualProductionShortCategoryName", "VP");
 		IPlacementModeModule::Get().RegisterPlacementCategory(Info);
 
 		return PlacementModeModule.GetRegisteredPlacementCategory(FLevelEditorOutlinerBuiltInCategories::VirtualProduction());
diff --git a/Engine/Plugins/VirtualProduction/CameraCalibrationCore/Source/CameraCalibrationCoreEditor/Private/CameraCalibrationCoreEditorModule.cpp b/Engine/Plugins/VirtualProduction/CameraCalibrationCore/Source/CameraCalibrationCoreEditor/Private/CameraCalibrationCoreEditorModule.cpp
--- a/Engine/Plugins/VirtualProduction/CameraCalibrationCore/Source/CameraCalibrationCoreEditor/Private/CameraCalibrationCoreEditorModule.cpp
+++ b/Engine/Plugins/VirtualProduction/CameraCalibrationCore/Source/CameraCalibrationCoreEditor/Private/CameraCalibrationCoreEditorModule.cpp
@@ -85,7 +85,7 @@ const FPlacementCategoryInfo* FCameraCalibrationCoreEditorModule::GetVirtualProd
 			TEXT("PMVirtualProduction"),
 			25 // Determines where the category shows up in the list with respect to the others.
 		);
-
+		Info.ShortDisplayName = LOCTEXT("VirtualProductionShortCategoryName", "VP");
 		IPlacementModeModule::Get().RegisterPlacementCategory(Info);
 
 		return PlacementModeModule.GetRegisteredPlacementCategory(FLevelEditorOutlinerBuiltInCategories::VirtualProduction());
diff --git a/Engine/Source/Developer/WidgetRegistration/Private/Layout/CategoryDrivenContentBuilder.cpp b/Engine/Source/Developer/WidgetRegistration/Private/Layout/CategoryDrivenContentBuilder.cpp
--- a/Engine/Source/Developer/WidgetRegistration/Private/Layout/CategoryDrivenContentBuilder.cpp
+++ b/Engine/Source/Developer/WidgetRegistration/Private/Layout/CategoryDrivenContentBuilder.cpp
@@ -234,8 +234,8 @@ void FCategoryDrivenContentBuilder::UpdateWidget()
 		{
 			ZeroStateBuilder = bIsFavoritesCategory ?
 				                   FZeroStateBuilderTemplates::Get().GetFavorites(
-					                   LOCTEXT("CategoryDrivenContentBuilder_NoFavoritesYetDragActorsToAdd",
-					                           "No favorites yet.\n\n To create favorites, drag actors from other categories and drop them on the Favorites star.")) :
+					                   LOCTEXT("CategoryDrivenContentBuilder_NoFavoritesYet",
+					                           "No favorites yet.\n\n To create favorites, right-click on items from other categories and add them to the Favorites.")) :
 
 				                   FZeroStateBuilderTemplates::Get().GetDefault(
 					                   LOCTEXT("CategoryDrivenContentBuilder_NoActorsMatchSearch",
diff --git a/Engine/Source/Developer/WidgetRegistration/Public/Inputs/BuilderInput.cpp b/Engine/Source/Developer/WidgetRegistration/Public/Inputs/BuilderInput.cpp
--- a/Engine/Source/Developer/WidgetRegistration/Public/Inputs/BuilderInput.cpp
+++ b/Engine/Source/Developer/WidgetRegistration/Public/Inputs/BuilderInput.cpp
@@ -37,9 +37,8 @@ UE::DisplayBuilders::FBuilderInput::FBuilderInput(
 	FInputChord InDefaultChords ,
 	FName InUiStyle ,
 	FName InBindingContext  ,
-	FName InBundle  ,
-	bool bInUseLongDisplayName )  :
-	FLabelAndIconArgs( InLabel.IsEmpty() ? FText::FromName( InName ) : InLabel, InIcon )
+	FName InBundle )
+	: FLabelAndIconArgs( InLabel.IsEmpty() ? FText::FromName( InName ) : InLabel, InIcon )
 	, Name( InName )
 	, UserInterfaceType( InUserInterfaceType )
 	, Description( InDescription.IsEmpty() ? Label : InDescription )
@@ -48,7 +47,6 @@ UE::DisplayBuilders::FBuilderInput::FBuilderInput(
 	, UIStyle( InUiStyle )
 	, BindingContext( InBindingContext )
 	, Bundle( InBundle )
-	, bUseLongDisplayName( bInUseLongDisplayName )
 	, Index( INDEX_NONE )
 	, Tooltip( InToolTip.IsEmpty() ? InLabel : InToolTip )
 {
diff --git a/Engine/Source/Developer/WidgetRegistration/Public/Inputs/BuilderInput.h b/Engine/Source/Developer/WidgetRegistration/Public/Inputs/BuilderInput.h
--- a/Engine/Source/Developer/WidgetRegistration/Public/Inputs/BuilderInput.h
+++ b/Engine/Source/Developer/WidgetRegistration/Public/Inputs/BuilderInput.h
@@ -33,8 +33,7 @@ namespace UE::DisplayBuilders
 			FInputChord InDefaultChords = FInputChord(),
 			FName InUiStyle = NAME_None,
 			FName InBindingContext = NAME_None,
-			FName InBundle = NAME_None,
-			bool bInUseLongDisplayName = false );
+			FName InBundle = NAME_None );
 
 		/**
 		 * The destructor, destroys any command related information
@@ -65,9 +64,6 @@ namespace UE::DisplayBuilders
 		/** The bundle to group the UI command into. The bundle must have been added to the BindingContext first. */
 		FName Bundle;
 
-		/** True if should we use long names for when getting text for input chords */
-		bool bUseLongDisplayName;
-
 		/** the index of the BuilderInput in whatever container it is in */
 		int32 Index;
 		
diff --git a/Engine/Source/Editor/PlacementMode/Private/PlacementModeModule.cpp b/Engine/Source/Editor/PlacementMode/Private/PlacementModeModule.cpp
--- a/Engine/Source/Editor/PlacementMode/Private/PlacementModeModule.cpp
+++ b/Engine/Source/Editor/PlacementMode/Private/PlacementModeModule.cpp
@@ -239,15 +239,16 @@ void FPlacementModeModule::StartupModule()
 	{
 		int32 SortOrder = 0;
 		FName CategoryName = FBuiltInPlacementCategories::Visual();
-		RegisterPlacementCategory(
-			FPlacementCategoryInfo(
-				NSLOCTEXT("PlacementMode", "VisualEffects", "VFX"),
+
+		FPlacementCategoryInfo VfxCategoryInfo = FPlacementCategoryInfo(
+				NSLOCTEXT("PlacementMode", "VisualEffects", "Visual Effects"),
 				FSlateIcon(FAppStyle::GetAppStyleSetName(), "PlacementBrowser.Icons.VisualEffects"),
 				CategoryName,
 				TEXT("PMVisual"),
 				30
-			)
-		);
+			);
+		VfxCategoryInfo.ShortDisplayName = NSLOCTEXT("PlacementMode", "VisualEffectsShortCategoryName", "VFX");
+		RegisterPlacementCategory(MoveTemp(VfxCategoryInfo));
 
 		UActorFactory* PPFactory = GEditor->FindActorFactoryByClassForActorClass(UActorFactoryBoxVolume::StaticClass(), APostProcessVolume::StaticClass());
 
diff --git a/Engine/Source/Editor/PlacementMode/Private/SPlacementModeTools.cpp b/Engine/Source/Editor/PlacementMode/Private/SPlacementModeTools.cpp
--- a/Engine/Source/Editor/PlacementMode/Private/SPlacementModeTools.cpp
+++ b/Engine/Source/Editor/PlacementMode/Private/SPlacementModeTools.cpp
@@ -481,7 +481,7 @@ void SPlacementAssetMenuEntry::Construct(const FArguments& InArgs, const TShared
 		.BorderImage( this, &SPlacementAssetMenuEntry::GetBorder )
 		.Cursor( EMouseCursor::GrabHand )
 		.ToolTip( AssetEntryToolTip )
-		.Padding(FMargin(27.f, 3.f, 5.f, 3.f))
+		.Padding(FMargin(10.f, 3.f, 5.f, 3.f))
 		[
 			SNew( SHorizontalBox )
 
@@ -530,13 +530,16 @@ void SPlacementAssetMenuEntry::Construct(const FArguments& InArgs, const TShared
 
 const FSlateBrush* SPlacementAssetMenuEntry::GetIcon() const
 {
-
 	if (AssetImage != nullptr)
 	{
 		return AssetImage;
 	}
 
-	if (Item->ClassIconBrushOverride != NAME_None)
+	if (Item->DragHandler && Item->DragHandler->IconBrush)
+	{
+		AssetImage = Item->DragHandler->IconBrush;
+	}
+	else if (Item->ClassIconBrushOverride != NAME_None)
 	{
 		AssetImage = FSlateIconFinder::FindCustomIconBrushForClass(nullptr, TEXT("ClassIcon"), Item->ClassIconBrushOverride);
 	}
@@ -618,6 +621,11 @@ FReply SPlacementAssetMenuEntry::OnDragDetected(const FGeometry& MyGeometry, con
 
 	if( MouseEvent.IsMouseButtonDown( EKeys::LeftMouseButton ) )
 	{
+		if (Item->DragHandler.IsValid() && Item->DragHandler->GetContentToDrag.IsBound())
+		{
+			return FReply::Handled().BeginDragDrop( Item->DragHandler->GetContentToDrag.Execute() );
+		}
+
 		return FReply::Handled().BeginDragDrop(FAssetDragDropOp::New(Item->AssetData, Item->AssetFactory));
 	}
 	else
@@ -1060,8 +1068,14 @@ void SPlacementModeTools::UpdatePlacementCategories()
 
 	for (const FPlacementCategoryInfo& Category : Categories)
 	{
-		const UE::DisplayBuilders::FBuilderInput& InputInfo = UE::DisplayBuilders::FBuilderInput( Category.UniqueHandle, Category.DisplayName,
-		Category.DisplayIcon, EUserInterfaceActionType::ToggleButton );
+		UE::DisplayBuilders::FBuilderInput InputInfo = UE::DisplayBuilders::FBuilderInput( Category.UniqueHandle, Category.DisplayName,
+			Category.DisplayIcon, EUserInterfaceActionType::ToggleButton );
+
+		if (!Category.ShortDisplayName.IsEmpty())
+		{
+			InputInfo.ButtonArgs.LabelOverride = Category.ShortDisplayName;
+		}
+
 		BuilderInputArray.Add( InputInfo );
 		
 		if ( Category.UniqueHandle == FBuiltInPlacementCategories::Favorites() && !CVarEnableCategoryContentChooserView->GetBool() )
diff --git a/Engine/Source/Editor/PlacementMode/Public/IPlacementModeModule.h b/Engine/Source/Editor/PlacementMode/Public/IPlacementModeModule.h
--- a/Engine/Source/Editor/PlacementMode/Public/IPlacementModeModule.h
+++ b/Engine/Source/Editor/PlacementMode/Public/IPlacementModeModule.h
@@ -16,7 +16,6 @@
 #include "Textures/SlateIcon.h"
 
 class FNamePermissionList;
-class FBuilderInput;
 struct FPlaceableItem;
 
 /**
@@ -60,6 +59,9 @@ struct FPlacementCategoryInfo
 	/** This category's display name */
 	FText DisplayName;
 
+	/** This category's short display name (optional) */
+	FText ShortDisplayName;
+
 	/** This category's representative icon */
 	FSlateIcon DisplayIcon;
 
