Place Actors
- Removed the classic view and the cvar to switch between the two views.

diff --git a/Engine/Source/Editor/PlacementMode/Private/SPlacementModeTools.cpp b/Engine/Source/Editor/PlacementMode/Private/SPlacementModeTools.cpp
--- a/Engine/Source/Editor/PlacementMode/Private/SPlacementModeTools.cpp
+++ b/Engine/Source/Editor/PlacementMode/Private/SPlacementModeTools.cpp
@@ -33,10 +33,6 @@
 
 #define LOCTEXT_NAMESPACE "PlacementMode"
 
-static TAutoConsoleVariable<bool> CVarEnableCategoryContentChooserView(
-	TEXT("PlaceActors.EnableCategoryChooserView"),
-	false,
-	TEXT("If enabled, the new category chooser view will be used for the Place Actors panel."));
 
 namespace PlacementModeTools
 {
@@ -130,33 +126,22 @@ public:
 
 		Thumbnail = MakeShareable(new FAssetThumbnail(Asset, InArgs._Width, InArgs._Height, ThumbnailPool));
 		
-		if ( CVarEnableCategoryContentChooserView->GetBool() )
-		{
-			TSharedPtr<SImage> ThumbnailImage;
-
-			// figure out the proper image to show based on whether the asset is a class type
-			TWeakObjectPtr<UClass> ThumbnailClass = MakeWeakObjectPtr( const_cast<UClass*>( FClassIconFinder::GetIconClassForAssetData( Asset, &bIsClassType ) ) );
-			const FName AssetClassName = Asset.AssetClassPath.GetAssetName();
-			const FName DefaultThumbnail = bIsClassType ? NAME_None : FName( *FString::Printf( TEXT("ClassThumbnail.%s"), *AssetClassName.ToString() ) );
-			const FSlateBrush* ThumbnailBrush = !InArgs._ClassThumbnailBrushOverride.IsNone() ?
-				FClassIconFinder::FindThumbnailForClass( nullptr,  InArgs._ClassThumbnailBrushOverride ) :
-				FClassIconFinder::FindThumbnailForClass( ThumbnailClass.Get(), DefaultThumbnail );
+		TSharedPtr<SImage> ThumbnailImage;
 
-			if ( InArgs._CustomIconBrush )
-			{
-				ThumbnailBrush = InArgs._CustomIconBrush;
-			}
+		// figure out the proper image to show based on whether the asset is a class type
+		TWeakObjectPtr<UClass> ThumbnailClass = MakeWeakObjectPtr( const_cast<UClass*>( FClassIconFinder::GetIconClassForAssetData( Asset, &bIsClassType ) ) );
+		const FName AssetClassName = Asset.AssetClassPath.GetAssetName();
+		const FName DefaultThumbnail = bIsClassType ? NAME_None : FName( *FString::Printf( TEXT("ClassThumbnail.%s"), *AssetClassName.ToString() ) );
+		const FSlateBrush* ThumbnailBrush = !InArgs._ClassThumbnailBrushOverride.IsNone() ?
+			FClassIconFinder::FindThumbnailForClass( nullptr,  InArgs._ClassThumbnailBrushOverride ) :
+			FClassIconFinder::FindThumbnailForClass( ThumbnailClass.Get(), DefaultThumbnail );
 
-			ChildSlot[ SAssignNew( ThumbnailImage, SImage ).Image( ThumbnailBrush ) ];
-		}
-		else
+		if ( InArgs._CustomIconBrush )
 		{
-			FAssetThumbnailConfig Config;
-			Config.bForceGenericThumbnail = InArgs._AlwaysUseGenericThumbnail;
-			Config.ClassThumbnailBrushOverride = InArgs._ClassThumbnailBrushOverride;
-			Config.AssetTypeColorOverride = InArgs._AssetTypeColorOverride;
-			ChildSlot[ Thumbnail->MakeThumbnailWidget( Config ) ];			
+			ThumbnailBrush = InArgs._CustomIconBrush;
 		}
+
+		ChildSlot[ SAssignNew( ThumbnailImage, SImage ).Image( ThumbnailBrush ) ];
 	}
 
 private:
@@ -220,7 +205,6 @@ void SPlacementAssetEntry::Construct(const FArguments& InArgs, const TSharedPtr<
 	NormalImage = &ButtonStyle.Normal;
 	HoverImage = &ButtonStyle.Hovered;
 	PressedImage = &ButtonStyle.Pressed;
-	const bool bIsClassicView = !CVarEnableCategoryContentChooserView->GetBool();
 	float ThumbnailBoxWidth = 40;
 	
 	float TextFillWidth = 0.99;
@@ -231,32 +215,18 @@ void SPlacementAssetEntry::Construct(const FArguments& InArgs, const TSharedPtr<
 	FMargin ThumbnailBoxPadding{ 8.f ,4.f,8.f, 4.f };
 	FMargin AssetTextPadding{ 9, 0, 0, 1 };
 	TSharedRef<SWidget> DraggableAssetEndWidget = SNullWidget::NullWidget;
-	
-	if ( bIsClassicView )
-	{
-		// Create doc link widget if there is a class to link to
-		TSharedRef<SWidget> DocWidget = SNew( SSpacer );
-		if( DocClass != NULL )
-		{
-			DocWidget = FEditorClassUtils::GetDocumentationLinkWidget( DocClass );
-			DocWidget->SetCursor( EMouseCursor::Default );
-		}
-		DraggableAssetEndWidget = DocWidget;
-	}
-	else
-	{
-		WholeAssetPadding = 0;
-		WholeAssetBackgroundBrush = FAppStyle::Get().GetBrush("PlacementBrowser.Asset.ThumbnailBackground");
-		ThumbnailBoxPadding = FMargin{ 4.f,4.f,0.f, 4.f };
-		AssetTextPadding = FMargin{ 4, 0, 8, 1 };
-		DraggableAssetEndWidget = SNew(SBox )
-									.Padding( DragHandlePadding )
-									[  SNew(SImage).Image( FSlateBrushTemplates::DragHandle() ) ];
-		ThumbnailBoxWidth = 20;
-	}
+
+	WholeAssetPadding = 0;
+	WholeAssetBackgroundBrush = FAppStyle::Get().GetBrush("PlacementBrowser.Asset.ThumbnailBackground");
+	ThumbnailBoxPadding = FMargin{ 4.f,4.f,0.f, 4.f };
+	AssetTextPadding = FMargin{ 4, 0, 8, 1 };
+	DraggableAssetEndWidget = SNew(SBox )
+								.Padding( DragHandlePadding )
+								[  SNew(SImage).Image( FSlateBrushTemplates::DragHandle() ) ];
+	ThumbnailBoxWidth = 20;
 	
 	const FSlateBrush* CustomIconBrush = nullptr;
-	if ( !bIsClassicView && Item->DragHandler.IsValid() && Item->DragHandler->IconBrush )
+	if ( Item->DragHandler.IsValid() && Item->DragHandler->IconBrush )
 	{
 		CustomIconBrush = Item->DragHandler->IconBrush;
 	}
@@ -345,7 +315,7 @@ FReply SPlacementAssetEntry::OnMouseButtonDown(const FGeometry& MyGeometry, cons
 	}
 
 	// Create the context menu to be launched on right mouse click. 
-	if ( CVarEnableCategoryContentChooserView->GetBool() && MouseEvent.GetEffectingButton() == EKeys::RightMouseButton )
+	if ( MouseEvent.GetEffectingButton() == EKeys::RightMouseButton )
 	{
 		FWidgetPath WidgetPath = MouseEvent.GetEventPath() != nullptr ? *MouseEvent.GetEventPath() : FWidgetPath();
 
@@ -385,9 +355,9 @@ FReply SPlacementAssetEntry::OnDragDetected(const FGeometry& MyGeometry, const F
 		return FReply::Handled();
 	}
 
-	if( MouseEvent.IsMouseButtonDown( EKeys::LeftMouseButton ) )
+	if ( MouseEvent.IsMouseButtonDown( EKeys::LeftMouseButton ) )
 	{
-		if ( Item->DragHandler.IsValid()  && Item->DragHandler->GetContentToDrag.IsBound())
+		if ( Item->DragHandler.IsValid() && Item->DragHandler->GetContentToDrag.IsBound() )
 		{
 			return FReply::Handled().BeginDragDrop( Item->DragHandler->GetContentToDrag.Execute() );
 		}
@@ -736,86 +706,11 @@ void SPlacementModeTools::Construct( const FArguments& InArgs, TSharedRef<SDockT
 			[
 				SNew(SBorder)
 				.BorderImage( FSlateBrushTemplates::Panel() )
-				.Visibility_Lambda( [] () { return CVarEnableCategoryContentChooserView->GetBool()   ? EVisibility::Visible : EVisibility::Collapsed; })
 				.Padding(0)
 				[
 					CategoryContentBuilder->GenerateWidgetSharedRef()
 				]
 			]
-		+ SVerticalBox::Slot()
-		.AutoHeight()
-		.HAlign(HAlign_Fill)
-		[
-			SNew(SBorder)
-			.BorderImage(FAppStyle::GetBrush("ToolPanel.GroupBorder"))
-			.Padding(FMargin(8.f, 6.f, 8.f, 8.f))
-			.HAlign(HAlign_Fill)
-			.Visibility(this, &SPlacementModeTools::GetTabsVisibility)
-			[
-				CategoryFilterPtr.ToSharedRef()
-			]
-		]
-
-
-		+ SVerticalBox::Slot()
-		.AutoHeight()
-		[
-			SNew(SBorder)
-			.BorderImage(FAppStyle::GetBrush("ToolPanel.GroupBorder"))
-			.Padding(FMargin(8.f, 6.f, 8.f, 8.f))
-			.HAlign(HAlign_Center)
-			.Visibility(this, &SPlacementModeTools::GetTabsVisibility)
-			[
-				SAssignNew(FilterLabelPtr, STextBlock)
-				.Text(LOCTEXT("CategoryLabel", "CategoryLabel"))
-				.Font(FAppStyle::Get().GetFontStyle("SmallFontBold"))
-				.TransformPolicy(ETextTransformPolicy::ToUpper)
-			]
-		]
-
-		+ SVerticalBox::Slot()
-		.Padding(FMargin(0.0f, 3.f))
-		[
-			SNew(SOverlay)
-			.Visibility(CVarEnableCategoryContentChooserView->GetBool() ? EVisibility::Collapsed : EVisibility::Visible)
-
-			+ SOverlay::Slot()
-			.HAlign(HAlign_Center)
-			.VAlign(VAlign_Fill)
-			.Padding(12.f)
-			[
-				SNew(STextBlock)
-				.Text(LOCTEXT("NoResultsFound", "No Results Found"))
-				.Visibility(this, &SPlacementModeTools::GetFailedSearchVisibility)
-			]
-
-			+ SOverlay::Slot()
-			[
-				SAssignNew(CustomContent, SBox)
-			]
-
-			+ SOverlay::Slot()
-			[
-				SAssignNew(DataDrivenContent, SBox)
-				[
-					SNew(SHorizontalBox)
-
-					+ SHorizontalBox::Slot()
-					[
-						SAssignNew(ListView, SListView<TSharedPtr<FPlaceableItem>>)
-						.ListItemsSource(&FilteredItems)
-						.OnGenerateRow(this, &SPlacementModeTools::OnGenerateWidgetForItem)
-						.ExternalScrollbar(ScrollBar)
-					]
-
-					+ SHorizontalBox::Slot()
-					.AutoWidth()
-					[
-						ScrollBar
-					]
-				]
-			]
-		]
 	];
 
 	IPlacementModeModule& PlacementModeModule = IPlacementModeModule::Get();
@@ -843,7 +738,6 @@ void SPlacementModeTools::SetActiveTab(FName TabName)
 void SPlacementModeTools::UpdateShownItems()
 {
 	bUpdateShownItems = false;
-	bool bIsCategoryContentChooserView = CVarEnableCategoryContentChooserView->GetBool();
 
 	IPlacementModeModule& PlacementModeModule = IPlacementModeModule::Get();
 
@@ -852,7 +746,7 @@ void SPlacementModeTools::UpdateShownItems()
 	{
 		return;
 	}
-	else if (Category->CustomGenerator && (Category->CustomDraggableItems.IsEmpty() || !bIsCategoryContentChooserView))
+	else if (Category->CustomGenerator && Category->CustomDraggableItems.IsEmpty())
 	{
 		CategoryContentBuilder->FillWithBuilder( Category->CustomGenerator() );
 		CustomContent->SetContent(Category->CustomGenerator());
@@ -908,14 +802,6 @@ void SPlacementModeTools::UpdateShownItems()
 				}
 			}
 		}
-
-		if ( !CVarEnableCategoryContentChooserView->GetBool() )
-		{
-			CustomContent->SetVisibility(EVisibility::Collapsed);
-			DataDrivenContent->SetVisibility(EVisibility::Visible);
-			ListView->RequestListRefresh();
-			FilterLabelPtr->SetText(Category->DisplayName);
-		}
 	}
 }
 
@@ -928,29 +814,6 @@ ECheckBoxState SPlacementModeTools::GetPlacementTabCheckedState( FName CategoryN
 {
 	return ActiveTabName == CategoryName ? ECheckBoxState::Checked : ECheckBoxState::Unchecked;
 }
-	
-EVisibility SPlacementModeTools::GetFailedSearchVisibility() const
-{
-	if (!IsSearchActive() || FilteredItems.Num())
-	{
-		return EVisibility::Collapsed;
-	}
-	return EVisibility::Visible;
-}
-
-EVisibility SPlacementModeTools::GetTabsVisibility() const
-{
-	return IsSearchActive() || CVarEnableCategoryContentChooserView->GetBool() ? EVisibility::Collapsed : EVisibility::Visible;
-}
-
-TSharedRef<ITableRow> SPlacementModeTools::OnGenerateWidgetForItem(TSharedPtr<FPlaceableItem> InItem, const TSharedRef<STableViewBase>& OwnerTable)
-{
-	return SNew(STableRow<TSharedPtr<FPlaceableItem>>, OwnerTable)
-		.Style(&FAppStyle::Get(), "PlacementBrowser.PlaceableItemRow")
-		[
-			GetPlacementAssetWidget(InItem)
-		];
-}
 
 TSharedRef<SWidget> SPlacementModeTools::GetPlacementAssetWidget( const TSharedPtr<FPlaceableItem>& InItem ) const
 {
@@ -966,43 +829,40 @@ TSharedRef<SWidget> SPlacementModeTools::GetPlacementAssetWidget( const TSharedP
 
 void SPlacementModeTools::UpdateContentForCategory( FName CategoryName, FText CategoryLabel )
 {
-	if ( CVarEnableCategoryContentChooserView->GetBool() )
-	{
-		SetActiveTab( CategoryName );
-		CustomContent = SNew( SBox );
-		FavoriteItems.Empty();
+	SetActiveTab( CategoryName );
+	CustomContent = SNew( SBox );
+	FavoriteItems.Empty();
 
-		CategoryContentBuilder->ClearCategoryContent();
+	CategoryContentBuilder->ClearCategoryContent();
 		
-		// if the Category name is not none, the user updated the category, so clear out the search ~ the Category choice should override it.
-		// The call of UpdateShownItems below will update search state based on this setting.
-		if ( !CategoryName.IsNone() )
-		{
-			SearchBoxPtr->SetText( FText::GetEmpty());
-		}
+	// if the Category name is not none, the user updated the category, so clear out the search ~ the Category choice should override it.
+	// The call of UpdateShownItems below will update search state based on this setting.
+	if ( !CategoryName.IsNone() )
+	{
+		SearchBoxPtr->SetText( FText::GetEmpty());
+	}
 		
-		UpdateShownItems();
+	UpdateShownItems();
 
-		const FPlacementCategoryInfo* Category = IPlacementModeModule::Get().GetRegisteredPlacementCategory( CategoryName );
-		if ( Category && Category->CustomGenerator && Category->CustomDraggableItems.IsEmpty() )
-		{
-			CategoryContentBuilder->FillWithBuilder( Category->CustomGenerator() );
-			return;
-		}
+	const FPlacementCategoryInfo* Category = IPlacementModeModule::Get().GetRegisteredPlacementCategory( CategoryName );
+	if ( Category && Category->CustomGenerator && Category->CustomDraggableItems.IsEmpty() )
+	{
+		CategoryContentBuilder->FillWithBuilder( Category->CustomGenerator() );
+		return;
+	}
 
-		 if ( IsFavoritesCategorySelected() )
+	if ( IsFavoritesCategorySelected() )
+	{
+		for ( const TSharedPtr<FPlaceableItem>& Item : FavoriteItems )
 		{
-				for ( const TSharedPtr<FPlaceableItem>& Item : FavoriteItems )
-				{
-					CategoryContentBuilder->AddBuilder( GetPlacementAssetWidget( Item ) );
-				}				
-		}
-		else
+			CategoryContentBuilder->AddBuilder( GetPlacementAssetWidget( Item ) );
+		}				
+	}
+	else
+	{
+		for (const TSharedPtr<FPlaceableItem>& Item : FilteredItems)
 		{
-			for (const TSharedPtr<FPlaceableItem>& Item : FilteredItems)
-			{
-				CategoryContentBuilder->AddBuilder( GetPlacementAssetWidget(Item) );
-			}
+			CategoryContentBuilder->AddBuilder( GetPlacementAssetWidget(Item) );
 		}
 	}
 }
@@ -1078,10 +938,6 @@ void SPlacementModeTools::UpdatePlacementCategories()
 
 		BuilderInputArray.Add( InputInfo );
 		
-		if ( Category.UniqueHandle == FBuiltInPlacementCategories::Favorites() && !CVarEnableCategoryContentChooserView->GetBool() )
-		{
-			continue;
-		}
 		if (Category.UniqueHandle == FBuiltInPlacementCategories::Basic())
 		{
 			bBasicTabExists = true;
diff --git a/Engine/Source/Editor/PlacementMode/Private/SPlacementModeTools.h b/Engine/Source/Editor/PlacementMode/Private/SPlacementModeTools.h
--- a/Engine/Source/Editor/PlacementMode/Private/SPlacementModeTools.h
+++ b/Engine/Source/Editor/PlacementMode/Private/SPlacementModeTools.h
@@ -122,10 +122,6 @@ private:
 	// End SWidget
 
 private:
-
-	/** Generates a widget for the specified item */
-	TSharedRef<ITableRow> OnGenerateWidgetForItem(TSharedPtr<FPlaceableItem> InItem, const TSharedRef<STableViewBase>& OwnerTable);
-
 	/** Get the identifier of the currently active tab */
 	FName GetActiveTab() const;
 
@@ -167,15 +163,6 @@ private:
 	/** Gets the tab 'active' state, so that we can show the active style */
 	ECheckBoxState GetPlacementTabCheckedState( FName CategoryName ) const;
 
-	/** Gets the visibility for the failed search text */
-	EVisibility GetFailedSearchVisibility() const;
-
-	/** Gets the visibility for the list view */
-	EVisibility GetListViewVisibility() const;
-
-	/** Gets the visibility for tabs */
-	EVisibility GetTabsVisibility() const;
-
 private:
 
 	/** Called when the search text changes */
