Place Actors Panel
- Fixed entries being duplicated when clicking on a category while a search is active.
- The cause is that we call UpdateWidget from within UpdateWidget when we clear the search field.

diff --git a/Engine/Source/Editor/PlacementMode/Private/SPlacementModeTools.cpp b/Engine/Source/Editor/PlacementMode/Private/SPlacementModeTools.cpp
--- a/Engine/Source/Editor/PlacementMode/Private/SPlacementModeTools.cpp
+++ b/Engine/Source/Editor/PlacementMode/Private/SPlacementModeTools.cpp
@@ -659,6 +659,7 @@ void SPlacementModeTools::Construct( const FArguments& InArgs, TSharedRef<SDockT
 	bRefreshAllClasses = false;
 	bRefreshRecentlyPlaced = false;
 	bUpdateShownItems = true;
+	bIsRawSearchChange = false;
 
 	FCategoryDrivenContentBuilderArgs Args( "PlacementModes", UE::DisplayBuilders::FBuilderKeys::Get().PlaceActors() );
 	Args.FavoritesCommandName = FBuiltInPlacementCategories::Favorites();
@@ -834,14 +835,15 @@ void SPlacementModeTools::UpdateContentForCategory( FName CategoryName, FText Ca
 	FavoriteItems.Empty();
 
 	CategoryContentBuilder->ClearCategoryContent();
-		
+
 	// if the Category name is not none, the user updated the category, so clear out the search ~ the Category choice should override it.
 	// The call of UpdateShownItems below will update search state based on this setting.
 	if ( !CategoryName.IsNone() )
 	{
-		SearchBoxPtr->SetText( FText::GetEmpty());
+		TGuardValue<bool> IsRawSearchChangeGuard(bIsRawSearchChange, true);
+		SearchBoxPtr->SetText( FText::GetEmpty() );
 	}
-		
+
 	UpdateShownItems();
 
 	const FPlacementCategoryInfo* Category = IPlacementModeModule::Get().GetRegisteredPlacementCategory( CategoryName );
@@ -856,7 +858,7 @@ void SPlacementModeTools::UpdateContentForCategory( FName CategoryName, FText Ca
 		for ( const TSharedPtr<FPlaceableItem>& Item : FavoriteItems )
 		{
 			CategoryContentBuilder->AddBuilder( GetPlacementAssetWidget( Item ) );
-		}				
+		}
 	}
 	else
 	{
@@ -1017,7 +1019,7 @@ void SPlacementModeTools::OnSearchChanged(const FText& InFilterText)
 	SearchTextFilter->SetRawFilterText( InFilterText );
 	SearchBoxPtr->SetError( SearchTextFilter->GetFilterErrorText() );
 
-	if ( !OldText.EqualToCaseIgnored( InFilterText ) )
+	if ( !OldText.EqualToCaseIgnored( InFilterText ) && !bIsRawSearchChange )
 	{
 		CategoryContentBuilder->SetShowNoCategorySelection( IsSearchActive() );
 		CategoryContentBuilder->UpdateWidget();
diff --git a/Engine/Source/Editor/PlacementMode/Private/SPlacementModeTools.h b/Engine/Source/Editor/PlacementMode/Private/SPlacementModeTools.h
--- a/Engine/Source/Editor/PlacementMode/Private/SPlacementModeTools.h
+++ b/Engine/Source/Editor/PlacementMode/Private/SPlacementModeTools.h
@@ -180,6 +180,9 @@ private:
 	/** Flag to indicate that we need to update the list of shown items */
 	bool bUpdateShownItems;
 
+	/** Flag to indicate that we're programmatically changing the search text and that we should skip updating the placement entries */
+	bool bIsRawSearchChange;
+
 	// The text filter used to filter the classes
 	typedef TTextFilter<const FPlaceableItem&> FPlacementAssetEntryTextFilter;
 	TSharedPtr<FPlacementAssetEntryTextFilter> SearchTextFilter;
