Place Actors panel
- Fixed a crash with categories with custom content. The cleanup that removed the old UI left some remenants behind that were null. Those have now been removed.

diff --git a/Engine/Source/Editor/PlacementMode/Private/SPlacementModeTools.cpp b/Engine/Source/Editor/PlacementMode/Private/SPlacementModeTools.cpp
--- a/Engine/Source/Editor/PlacementMode/Private/SPlacementModeTools.cpp
+++ b/Engine/Source/Editor/PlacementMode/Private/SPlacementModeTools.cpp
@@ -676,10 +676,6 @@ void SPlacementModeTools::Construct( const FArguments& InArgs, TSharedRef<SDockT
 		FPlacementAssetEntryTextFilter::FItemToStringArray::CreateStatic(&PlacementViewFilter::GetBasicStrings)
 		));
 
-	SAssignNew(CategoryFilterPtr, SUniformWrapPanel)
-	.HAlign(HAlign_Center)
-	.SlotPadding(FMargin(2.0f, 1.0f));
-
 	UpdatePlacementCategories();
 
 	TSharedRef<SScrollBar> ScrollBar = SNew(SScrollBar)
@@ -750,12 +746,6 @@ void SPlacementModeTools::UpdateShownItems()
 	else if (Category->CustomGenerator && Category->CustomDraggableItems.IsEmpty())
 	{
 		CategoryContentBuilder->FillWithBuilder( Category->CustomGenerator() );
-		CustomContent->SetContent(Category->CustomGenerator());
-
-		CustomContent->SetVisibility(EVisibility::Visible);
-		DataDrivenContent->SetVisibility(EVisibility::Collapsed);
-
-		FilterLabelPtr->SetText(Category->DisplayName);
 	}
 	else if ( IsFavoritesCategorySelected() )
 	{
@@ -831,7 +821,6 @@ TSharedRef<SWidget> SPlacementModeTools::GetPlacementAssetWidget( const TSharedP
 void SPlacementModeTools::UpdateContentForCategory( FName CategoryName, FText CategoryLabel )
 {
 	SetActiveTab( CategoryName );
-	CustomContent = SNew( SBox );
 	FavoriteItems.Empty();
 
 	CategoryContentBuilder->ClearCategoryContent();
@@ -921,8 +910,6 @@ void SPlacementModeTools::UpdatePlacementCategories()
 	bool bBasicTabExists = false;
 	FName TabToActivate;
 
-	CategoryFilterPtr->ClearChildren();
-
 	TArray<FPlacementCategoryInfo> Categories;
 	IPlacementModeModule::Get().GetSortedCategories(Categories);
 
@@ -949,21 +936,6 @@ void SPlacementModeTools::UpdatePlacementCategories()
 		{
 			TabToActivate = ActiveTabName;
 		}
-
-		CategoryFilterPtr->AddSlot()
-		[
-			SNew(SCheckBox)
-			.Padding(FMargin(4.f, 4.f))
-			.Style( &FAppStyle::Get(),  "PaletteToolBar.Tab" )
-			.OnCheckStateChanged(this, &SPlacementModeTools::OnCategoryChanged, Category.UniqueHandle)
-			.IsChecked(this, &SPlacementModeTools::GetPlacementTabCheckedState, Category.UniqueHandle)
-			.ToolTipText(Category.DisplayName)
-			[
-				SNew(SImage)
-				.ColorAndOpacity(FSlateColor::UseForeground())
-				.Image(Category.DisplayIcon.GetIcon())
-			]
-		];
 	}
 	CategoryContentBuilder->InitializeCategoryButtons( BuilderInputArray );
 	
diff --git a/Engine/Source/Editor/PlacementMode/Private/SPlacementModeTools.h b/Engine/Source/Editor/PlacementMode/Private/SPlacementModeTools.h
--- a/Engine/Source/Editor/PlacementMode/Private/SPlacementModeTools.h
+++ b/Engine/Source/Editor/PlacementMode/Private/SPlacementModeTools.h
@@ -187,21 +187,9 @@ private:
 	typedef TTextFilter<const FPlaceableItem&> FPlacementAssetEntryTextFilter;
 	TSharedPtr<FPlacementAssetEntryTextFilter> SearchTextFilter;
 
-	/** Custom content slot, where a category has a custom generator */
-	TSharedPtr<SBox> CustomContent;
-	
-	/** Content container for any data driven content */
-	TSharedPtr<SBox> DataDrivenContent;
-
 	/* The search box used to update the filter text */
 	TSharedPtr<SSearchBox> SearchBoxPtr;
 
-	/* Category Filter */
-	TSharedPtr<SUniformWrapPanel> CategoryFilterPtr;
-
-	/* Active Category Filter Label */
-	TSharedPtr<STextBlock> FilterLabelPtr;
-
 	/** Array of filtered items to show in the list view */
 	TArray<TSharedPtr<FPlaceableItem>> FilteredItems;
 
