Add enhanced tooltips for Scene Graph Beta objects.
- entity class in the Place Actors panel and Create menu
- Entity Prefab Definition asset type in Content Browser

diff --git a/Engine/Source/Editor/PlacementMode/Private/SPlacementModeTools.cpp b/Engine/Source/Editor/PlacementMode/Private/SPlacementModeTools.cpp
--- a/Engine/Source/Editor/PlacementMode/Private/SPlacementModeTools.cpp
+++ b/Engine/Source/Editor/PlacementMode/Private/SPlacementModeTools.cpp
@@ -30,6 +30,7 @@
 #include "Styles/SlateBrushTemplates.h"
 #include "Layout/WidgetPath.h"
 #include "Styling/CoreStyle.h"
+#include "IDocumentation.h"
 
 #define LOCTEXT_NAMESPACE "PlacementMode"
 
@@ -197,7 +198,7 @@ void SPlacementAssetEntry::Construct(const FArguments& InArgs, const TSharedPtr<
 
 	if (!AssetEntryToolTip)
 	{
-		AssetEntryToolTip = FSlateApplicationBase::Get().MakeToolTip(Item->DisplayName);
+		AssetEntryToolTip = IDocumentation::Get()->CreateToolTip(Item->DisplayName, nullptr, "Shared/Types/AssetEntries", Item->DisplayName.ToString());
 	}
 
 	const FButtonStyle& ButtonStyle = FAppStyle::GetWidgetStyle<FButtonStyle>( "PlacementBrowser.Asset" );
@@ -427,7 +428,7 @@ void SPlacementAssetMenuEntry::Construct(const FArguments& InArgs, const TShared
 
 	if (!AssetEntryToolTip.IsValid())
 	{
-		AssetEntryToolTip = FSlateApplicationBase::Get().MakeToolTip(Item->DisplayName);
+		AssetEntryToolTip = IDocumentation::Get()->CreateToolTip(Item->DisplayName, nullptr, "Shared/Types/AssetEntries", Item->DisplayName.ToString());
 	}
 	
 	const FButtonStyle& ButtonStyle = FAppStyle::Get().GetWidgetStyle<FButtonStyle>( "Menu.Button" );
