Place Actors Panel: updated icon for Shapes and minor code cleanup

diff --git a/Engine/Source/Editor/EditorStyle/Private/StarshipStyle.cpp b/Engine/Source/Editor/EditorStyle/Private/StarshipStyle.cpp
--- a/Engine/Source/Editor/EditorStyle/Private/StarshipStyle.cpp
+++ b/Engine/Source/Editor/EditorStyle/Private/StarshipStyle.cpp
@@ -4856,7 +4856,8 @@ void FStarshipEditorStyle::FStyle::SetupLevelEditorStyle()
 		Set("PlacementBrowser.Icons.BSP",           new IMAGE_BRUSH_SVG("Starship/Common/Geometry",       Icon20x20));
 		Set("PlacementBrowser.Icons.Volumes",       new IMAGE_BRUSH_SVG("Starship/Common/Volumes",        Icon20x20));
 		Set("PlacementBrowser.Icons.All",           new IMAGE_BRUSH_SVG("Starship/Common/AllClasses",     Icon20x20));
-		Set("PlacementBrowser.Icons.Testing",       new CORE_IMAGE_BRUSH_SVG("Starship/Common/Test",           Icon20x20));
+		Set("PlacementBrowser.Icons.Testing",       new CORE_IMAGE_BRUSH_SVG("Starship/Common/Test",      Icon20x20));
+		Set("PlacementBrowser.Icons.Shapes",        new IMAGE_BRUSH_SVG("Starship/Common/Shapes",         Icon20x20));
 
 		Set( "ContentPalette.ShowAllPlaceables", new IMAGE_BRUSH( "Icons/icon_Placement_FilterAll_20px", Icon20x20 ) );
 		Set( "ContentPalette.ShowAllPlaceables.Small", new IMAGE_BRUSH( "Icons/icon_Placement_FilterAll_20px", Icon20x20 ) );
diff --git a/Engine/Source/Editor/PlacementMode/Private/PlacementModeModule.cpp b/Engine/Source/Editor/PlacementMode/Private/PlacementModeModule.cpp
--- a/Engine/Source/Editor/PlacementMode/Private/PlacementModeModule.cpp
+++ b/Engine/Source/Editor/PlacementMode/Private/PlacementModeModule.cpp
@@ -56,6 +56,8 @@
 #include "SPlacementModeTools.h"
 #include "AssetSelection.h"
 
+#define LOCTEXT_NAMESPACE "PlacementMode"
+
 namespace PlacementModeModuleLocals
 {
 	FActorPlacementInfo MakePlacementInfo(const FPlaceableItem& Item)
@@ -85,7 +87,7 @@ TOptional<FLinearColor> GetBasicShapeColorOverride()
 }
 
 FPlacementModeModule::FPlacementModeModule()
-	: CategoryPermissionList(MakeShareable(new FNamePermissionList()))
+	: CategoryPermissionList(MakeShared<FNamePermissionList>())
 {
 	CategoryPermissionList->OnFilterChanged().AddRaw(this, &FPlacementModeModule::OnCategoryPermissionListChanged);
 }
@@ -119,7 +121,7 @@ void FPlacementModeModule::StartupModule()
 
 	RegisterPlacementCategory(
 		FPlacementCategoryInfo(
-			NSLOCTEXT("PlacementMode", "Favorites", "Favorites"),
+			LOCTEXT("Favorites", "Favorites"),
 			FSlateIcon(FAppStyle::Get().GetStyleSetName(), "Icons.Favorites.Small"),
 			FBuiltInPlacementCategories::Favorites(),
 			TEXT("Favorites"),
@@ -129,7 +131,7 @@ void FPlacementModeModule::StartupModule()
 	
 	RegisterPlacementCategory(
 		FPlacementCategoryInfo(
-			NSLOCTEXT("PlacementMode", "RecentlyPlaced", "Recent"),
+			LOCTEXT("RecentlyPlaced", "Recent"),
 			FSlateIcon(FAppStyle::GetAppStyleSetName(), "PlacementBrowser.Icons.Recent"),
 			FBuiltInPlacementCategories::RecentlyPlaced(),
 			TEXT("PMRecentlyPlaced"),
@@ -143,7 +145,7 @@ void FPlacementModeModule::StartupModule()
 		FName CategoryName = FBuiltInPlacementCategories::Basic();
 		RegisterPlacementCategory(
 			FPlacementCategoryInfo(
-				NSLOCTEXT("PlacementMode", "Basic", "Basic"),
+				LOCTEXT("Basic", "Basic"),
 				FSlateIcon(FAppStyle::GetAppStyleSetName(), "PlacementBrowser.Icons.Basic"),
 				CategoryName,
 				TEXT("PMBasic"),
@@ -152,14 +154,14 @@ void FPlacementModeModule::StartupModule()
 		);
 
 		FPlacementCategory* Category = Categories.Find(CategoryName);
-		Category->Items.Add(CreateID(), MakeShareable(new FPlaceableItem(*UActorFactoryEmptyActor::StaticClass(), SortOrder += 10)));
-		Category->Items.Add(CreateID(), MakeShareable(new FPlaceableItem(*UActorFactoryCharacter::StaticClass(), SortOrder += 10)));
-		Category->Items.Add(CreateID(), MakeShareable(new FPlaceableItem(*UActorFactoryPawn::StaticClass(), SortOrder += 10)));
-		Category->Items.Add(CreateID(), MakeShareable(new FPlaceableItem(*UActorFactoryPointLight::StaticClass(), SortOrder += 10)));
-		Category->Items.Add(CreateID(), MakeShareable(new FPlaceableItem(*UActorFactoryPlayerStart::StaticClass(), SortOrder += 10)));
+		Category->Items.Add(CreateID(), MakeShared<FPlaceableItem>(*UActorFactoryEmptyActor::StaticClass(), SortOrder += 10));
+		Category->Items.Add(CreateID(), MakeShared<FPlaceableItem>(*UActorFactoryCharacter::StaticClass(), SortOrder += 10));
+		Category->Items.Add(CreateID(), MakeShared<FPlaceableItem>(*UActorFactoryPawn::StaticClass(), SortOrder += 10));
+		Category->Items.Add(CreateID(), MakeShared<FPlaceableItem>(*UActorFactoryPointLight::StaticClass(), SortOrder += 10));
+		Category->Items.Add(CreateID(), MakeShared<FPlaceableItem>(*UActorFactoryPlayerStart::StaticClass(), SortOrder += 10));
 	
-		Category->Items.Add(CreateID(), MakeShareable(new FPlaceableItem(*UActorFactoryTriggerBox::StaticClass(), SortOrder += 10)));
-		Category->Items.Add(CreateID(), MakeShareable(new FPlaceableItem(*UActorFactoryTriggerSphere::StaticClass(), SortOrder += 10)));
+		Category->Items.Add(CreateID(), MakeShared<FPlaceableItem>(*UActorFactoryTriggerBox::StaticClass(), SortOrder += 10));
+		Category->Items.Add(CreateID(), MakeShared<FPlaceableItem>(*UActorFactoryTriggerSphere::StaticClass(), SortOrder += 10));
 	}
 
 	{
@@ -167,7 +169,7 @@ void FPlacementModeModule::StartupModule()
 		FName CategoryName = FBuiltInPlacementCategories::Lights();
 		RegisterPlacementCategory(
 			FPlacementCategoryInfo(
-				NSLOCTEXT("PlacementMode", "Lights", "Lights"),
+				LOCTEXT("Lights", "Lights"),
 				FSlateIcon(FAppStyle::GetAppStyleSetName(), "PlacementBrowser.Icons.Lights"),
 				CategoryName,
 				TEXT("PMLights"),
@@ -176,11 +178,11 @@ void FPlacementModeModule::StartupModule()
 		);
 
 		FPlacementCategory* Category = Categories.Find(CategoryName);
-		Category->Items.Add(CreateID(), MakeShareable(new FPlaceableItem(*UActorFactoryDirectionalLight::StaticClass(), SortOrder += 10)));
-		Category->Items.Add(CreateID(), MakeShareable(new FPlaceableItem(*UActorFactoryPointLight::StaticClass(), SortOrder += 10)));
-		Category->Items.Add(CreateID(), MakeShareable(new FPlaceableItem(*UActorFactorySpotLight::StaticClass(), SortOrder += 10)));
-		Category->Items.Add(CreateID(), MakeShareable(new FPlaceableItem(*UActorFactoryRectLight::StaticClass(), SortOrder += 10)));
-		Category->Items.Add(CreateID(), MakeShareable(new FPlaceableItem(*UActorFactorySkyLight::StaticClass(), SortOrder += 10)));
+		Category->Items.Add(CreateID(), MakeShared<FPlaceableItem>(*UActorFactoryDirectionalLight::StaticClass(), SortOrder += 10));
+		Category->Items.Add(CreateID(), MakeShared<FPlaceableItem>(*UActorFactoryPointLight::StaticClass(), SortOrder += 10));
+		Category->Items.Add(CreateID(), MakeShared<FPlaceableItem>(*UActorFactorySpotLight::StaticClass(), SortOrder += 10));
+		Category->Items.Add(CreateID(), MakeShared<FPlaceableItem>(*UActorFactoryRectLight::StaticClass(), SortOrder += 10));
+		Category->Items.Add(CreateID(), MakeShared<FPlaceableItem>(*UActorFactorySkyLight::StaticClass(), SortOrder += 10));
 	}
 
 	{
@@ -189,48 +191,48 @@ void FPlacementModeModule::StartupModule()
 		FName CategoryName = FBuiltInPlacementCategories::Shapes();
 		RegisterPlacementCategory(
 			FPlacementCategoryInfo(
-				NSLOCTEXT("PlacementMode", "Shapes", "Shapes"),
-				FSlateIcon(FAppStyle::GetAppStyleSetName(), "ClassIcon.Cube"),
+				LOCTEXT("Shapes", "Shapes"),
+				FSlateIcon(FAppStyle::GetAppStyleSetName(), "PlacementBrowser.Icons.Shapes"),
 				CategoryName,
 				TEXT("PMShapes"),
 				25	
 				)
 			);
 
-		static const FText CubeText = NSLOCTEXT("PlacementMode", "Cube", "Cube");
-		static const FText SphereText = NSLOCTEXT("PlacementMode", "Sphere", "Sphere");
-		static const FText CylinderText = NSLOCTEXT("PlacementMode", "Cylinder", "Cylinder");
-		static const FText ConeText = NSLOCTEXT("PlacementMode", "Cone", "Cone");
-		static const FText PlaneText = NSLOCTEXT("PlacementMode", "Plane", "Plane");
+		static const FText CubeText = LOCTEXT("Cube", "Cube");
+		static const FText SphereText = LOCTEXT("Sphere", "Sphere");
+		static const FText CylinderText = LOCTEXT("Cylinder", "Cylinder");
+		static const FText ConeText = LOCTEXT("Cone", "Cone");
+		static const FText PlaneText = LOCTEXT("Plane", "Plane");
 
 		FPlacementCategory* Category = Categories.Find(CategoryName);
 		// Cube
 		{
-			TSharedPtr<FPlaceableItem> Cube = MakeShareable(new FPlaceableItem(*UActorFactoryBasicShape::StaticClass(), FAssetData(LoadObject<UStaticMesh>(nullptr, *UActorFactoryBasicShape::BasicCube.ToString())), FName("ClassThumbnail.Cube"), FName("ClassIcon.Cube"), BasicShapeColorOverride, SortOrder += 10, CubeText));
+			TSharedPtr<FPlaceableItem> Cube = MakeShared<FPlaceableItem>(*UActorFactoryBasicShape::StaticClass(), FAssetData(LoadObject<UStaticMesh>(nullptr, *UActorFactoryBasicShape::BasicCube.ToString())), FName("ClassThumbnail.Cube"), FName("ClassIcon.Cube"), BasicShapeColorOverride, SortOrder += 10, CubeText);
 			ManuallyCreatedPlaceableItems.Add(MakePlacementInfo(*Cube), Cube);
 			Category->Items.Add(CreateID(), Cube);
 		}
 		// Sphere
 		{
-			TSharedPtr<FPlaceableItem> Sphere = MakeShareable(new FPlaceableItem(*UActorFactoryBasicShape::StaticClass(), FAssetData(LoadObject<UStaticMesh>(nullptr, *UActorFactoryBasicShape::BasicSphere.ToString())), FName("ClassThumbnail.Sphere"), FName("ClassIcon.Sphere"), BasicShapeColorOverride, SortOrder += 10, SphereText));
+			TSharedPtr<FPlaceableItem> Sphere = MakeShared<FPlaceableItem>(*UActorFactoryBasicShape::StaticClass(), FAssetData(LoadObject<UStaticMesh>(nullptr, *UActorFactoryBasicShape::BasicSphere.ToString())), FName("ClassThumbnail.Sphere"), FName("ClassIcon.Sphere"), BasicShapeColorOverride, SortOrder += 10, SphereText);
 			ManuallyCreatedPlaceableItems.Add(MakePlacementInfo(*Sphere), Sphere);
 			Category->Items.Add(CreateID(), Sphere);
 		}
 		// Cylinder
 		{
-			TSharedPtr<FPlaceableItem> Cylinder = MakeShareable(new FPlaceableItem(*UActorFactoryBasicShape::StaticClass(), FAssetData(LoadObject<UStaticMesh>(nullptr, *UActorFactoryBasicShape::BasicCylinder.ToString())), FName("ClassThumbnail.Cylinder"), FName("ClassIcon.Cylinder"), BasicShapeColorOverride, SortOrder += 10, CylinderText));
+			TSharedPtr<FPlaceableItem> Cylinder = MakeShared<FPlaceableItem>(*UActorFactoryBasicShape::StaticClass(), FAssetData(LoadObject<UStaticMesh>(nullptr, *UActorFactoryBasicShape::BasicCylinder.ToString())), FName("ClassThumbnail.Cylinder"), FName("ClassIcon.Cylinder"), BasicShapeColorOverride, SortOrder += 10, CylinderText);
 			ManuallyCreatedPlaceableItems.Add(MakePlacementInfo(*Cylinder), Cylinder);
 			Category->Items.Add(CreateID(), Cylinder);
 		}
 		// Cone
 		{
-			TSharedPtr<FPlaceableItem> Cone = MakeShareable(new FPlaceableItem(*UActorFactoryBasicShape::StaticClass(), FAssetData(LoadObject<UStaticMesh>(nullptr, *UActorFactoryBasicShape::BasicCone.ToString())), FName("ClassThumbnail.Cone"), FName("ClassIcon.Cone"), BasicShapeColorOverride, SortOrder += 10, ConeText));
+			TSharedPtr<FPlaceableItem> Cone = MakeShared<FPlaceableItem>(*UActorFactoryBasicShape::StaticClass(), FAssetData(LoadObject<UStaticMesh>(nullptr, *UActorFactoryBasicShape::BasicCone.ToString())), FName("ClassThumbnail.Cone"), FName("ClassIcon.Cone"), BasicShapeColorOverride, SortOrder += 10, ConeText);
 			ManuallyCreatedPlaceableItems.Add(MakePlacementInfo(*Cone), Cone);
 			Category->Items.Add(CreateID(), Cone);
 		}
 		// Plane
 		{
-			TSharedPtr<FPlaceableItem> Plane = MakeShareable(new FPlaceableItem(*UActorFactoryBasicShape::StaticClass(), FAssetData(LoadObject<UStaticMesh>(nullptr, *UActorFactoryBasicShape::BasicPlane.ToString())), FName("ClassThumbnail.Plane"), FName("ClassIcon.Plane"), BasicShapeColorOverride, SortOrder += 10, PlaneText));
+			TSharedPtr<FPlaceableItem> Plane = MakeShared<FPlaceableItem>(*UActorFactoryBasicShape::StaticClass(), FAssetData(LoadObject<UStaticMesh>(nullptr, *UActorFactoryBasicShape::BasicPlane.ToString())), FName("ClassThumbnail.Plane"), FName("ClassIcon.Plane"), BasicShapeColorOverride, SortOrder += 10, PlaneText);
 			ManuallyCreatedPlaceableItems.Add(MakePlacementInfo(*Plane), Plane);
 			Category->Items.Add(CreateID(), Plane);
 		}
@@ -241,32 +243,32 @@ void FPlacementModeModule::StartupModule()
 		FName CategoryName = FBuiltInPlacementCategories::Visual();
 
 		FPlacementCategoryInfo VfxCategoryInfo = FPlacementCategoryInfo(
-				NSLOCTEXT("PlacementMode", "VisualEffects", "Visual Effects"),
+				LOCTEXT("VisualEffects", "Visual Effects"),
 				FSlateIcon(FAppStyle::GetAppStyleSetName(), "PlacementBrowser.Icons.VisualEffects"),
 				CategoryName,
 				TEXT("PMVisual"),
 				30
 			);
-		VfxCategoryInfo.ShortDisplayName = NSLOCTEXT("PlacementMode", "VisualEffectsShortCategoryName", "VFX");
+		VfxCategoryInfo.ShortDisplayName = LOCTEXT("VisualEffectsShortCategoryName", "VFX");
 		RegisterPlacementCategory(MoveTemp(VfxCategoryInfo));
 
 		UActorFactory* PPFactory = GEditor->FindActorFactoryByClassForActorClass(UActorFactoryBoxVolume::StaticClass(), APostProcessVolume::StaticClass());
 
 		FPlacementCategory* Category = Categories.Find(CategoryName);
-		Category->Items.Add(CreateID(), MakeShareable(new FPlaceableItem(PPFactory, FAssetData(APostProcessVolume::StaticClass()), SortOrder += 10)));
-		Category->Items.Add(CreateID(), MakeShareable(new FPlaceableItem(*UActorFactorySkyAtmosphere::StaticClass(), SortOrder += 10)));
-		Category->Items.Add(CreateID(), MakeShareable(new FPlaceableItem(*UActorFactoryVolumetricCloud::StaticClass(), SortOrder += 10)));
-		Category->Items.Add(CreateID(), MakeShareable(new FPlaceableItem(*UActorFactoryExponentialHeightFog::StaticClass(), SortOrder += 10)));
-		Category->Items.Add(CreateID(), MakeShareable(new FPlaceableItem(*UActorFactoryLocalFogVolume::StaticClass(), SortOrder += 10)));
-		Category->Items.Add(CreateID(), MakeShareable(new FPlaceableItem(*UActorFactorySphereReflectionCapture::StaticClass(), SortOrder += 10)));
-		Category->Items.Add(CreateID(), MakeShareable(new FPlaceableItem(*UActorFactoryBoxReflectionCapture::StaticClass(), SortOrder += 10)));
-		Category->Items.Add(CreateID(), MakeShareable(new FPlaceableItem(*UActorFactoryPlanarReflection::StaticClass(), SortOrder += 10)));
-		Category->Items.Add(CreateID(), MakeShareable(new FPlaceableItem(*UActorFactoryDeferredDecal::StaticClass(), SortOrder += 10)));
+		Category->Items.Add(CreateID(), MakeShared<FPlaceableItem>(PPFactory, FAssetData(APostProcessVolume::StaticClass()), SortOrder += 10));
+		Category->Items.Add(CreateID(), MakeShared<FPlaceableItem>(*UActorFactorySkyAtmosphere::StaticClass(), SortOrder += 10));
+		Category->Items.Add(CreateID(), MakeShared<FPlaceableItem>(*UActorFactoryVolumetricCloud::StaticClass(), SortOrder += 10));
+		Category->Items.Add(CreateID(), MakeShared<FPlaceableItem>(*UActorFactoryExponentialHeightFog::StaticClass(), SortOrder += 10));
+		Category->Items.Add(CreateID(), MakeShared<FPlaceableItem>(*UActorFactoryLocalFogVolume::StaticClass(), SortOrder += 10));
+		Category->Items.Add(CreateID(), MakeShared<FPlaceableItem>(*UActorFactorySphereReflectionCapture::StaticClass(), SortOrder += 10));
+		Category->Items.Add(CreateID(), MakeShared<FPlaceableItem>(*UActorFactoryBoxReflectionCapture::StaticClass(), SortOrder += 10));
+		Category->Items.Add(CreateID(), MakeShared<FPlaceableItem>(*UActorFactoryPlanarReflection::StaticClass(), SortOrder += 10));
+		Category->Items.Add(CreateID(), MakeShared<FPlaceableItem>(*UActorFactoryDeferredDecal::StaticClass(), SortOrder += 10));
 	}
 
 	RegisterPlacementCategory(
 		FPlacementCategoryInfo(
-			NSLOCTEXT("PlacementMode", "Volumes", "Volumes"),
+			LOCTEXT("Volumes", "Volumes"),
 			FSlateIcon(FAppStyle::GetAppStyleSetName(), "PlacementBrowser.Icons.Volumes"),
 			FBuiltInPlacementCategories::Volumes(),
 			TEXT("PMVolumes"),
@@ -276,7 +278,7 @@ void FPlacementModeModule::StartupModule()
 
 	RegisterPlacementCategory(
 		FPlacementCategoryInfo(
-			NSLOCTEXT("PlacementMode", "AllClasses", "All"),
+			LOCTEXT("AllClasses", "All"),
 			FSlateIcon(FAppStyle::GetAppStyleSetName(), "PlacementBrowser.Icons.All"),
 			FBuiltInPlacementCategories::AllClasses(),
 			TEXT("PMAllClasses"),
@@ -342,7 +344,7 @@ void FPlacementModeModule::StartupModule()
 		
 		UToolMenu* ContentMenu = UToolMenus::Get()->ExtendMenu("LevelEditor.LevelEditorToolBar.AddQuickMenu");
 		FName CreateSectionName = TEXT("PMQCreateMenu");
-		FText CreateSectionDisplayName = NSLOCTEXT("PlacementMode", "PMQCreateMenu", "Place Actors");
+		FText CreateSectionDisplayName = LOCTEXT("PMQCreateMenu", "Place Actors");
 
 		GenerateCategorySubMenu(ContentMenu, CreateSectionName, CreateSectionDisplayName, FBuiltInPlacementCategories::Basic());
 
@@ -764,7 +766,7 @@ void FPlacementModeModule::RefreshRecentlyPlaced()
 			continue;
 		}
 
-		TSharedPtr<FPlaceableItem> Ptr = MakeShareable(new FPlaceableItem(Factory, AssetData));
+		TSharedPtr<FPlaceableItem> Ptr = MakeShared<FPlaceableItem>(Factory, AssetData);
 		Category->Items.Add(CreateID(), Ptr);
 	}
 }
@@ -790,7 +792,7 @@ void FPlacementModeModule::RefreshVolumes()
 			Class->ClassGeneratedBy == nullptr)
 		{
 			UActorFactory* Factory = GEditor->FindActorFactoryByClassForActorClass(UActorFactoryBoxVolume::StaticClass(), Class);
-			Category->Items.Add(CreateID(), MakeShareable(new FPlaceableItem(Factory, FAssetData(Class))));
+			Category->Items.Add(CreateID(), MakeShared<FPlaceableItem>(Factory, FAssetData(Class)));
 		}
 	}
 }
@@ -816,14 +818,14 @@ void FPlacementModeModule::RefreshAllPlaceableClasses()
 	}
 
 	// Manually add some special cases that aren't added below
-	Category->Items.Add(CreateID(), MakeShareable(new FPlaceableItem(*UActorFactoryEmptyActor::StaticClass())));
-	Category->Items.Add(CreateID(), MakeShareable(new FPlaceableItem(*UActorFactoryCharacter::StaticClass())));
-	Category->Items.Add(CreateID(), MakeShareable(new FPlaceableItem(*UActorFactoryPawn::StaticClass())));
-	Category->Items.Add(CreateID(), MakeShareable(new FPlaceableItem(*UActorFactoryBasicShape::StaticClass(), FAssetData(LoadObject<UStaticMesh>(nullptr, *UActorFactoryBasicShape::BasicCube.ToString())), FName("ClassThumbnail.Cube"), FName("ClassIcon.Cube"), GetBasicShapeColorOverride(), TOptional<int32>(), NSLOCTEXT("PlacementMode", "Cube", "Cube"))));
-	Category->Items.Add(CreateID(), MakeShareable(new FPlaceableItem(*UActorFactoryBasicShape::StaticClass(), FAssetData(LoadObject<UStaticMesh>(nullptr, *UActorFactoryBasicShape::BasicSphere.ToString())), FName("ClassThumbnail.Sphere"), FName("ClassIcon.Sphere"), GetBasicShapeColorOverride(), TOptional<int32>(), NSLOCTEXT("PlacementMode", "Sphere", "Sphere"))));
-	Category->Items.Add(CreateID(), MakeShareable(new FPlaceableItem(*UActorFactoryBasicShape::StaticClass(), FAssetData(LoadObject<UStaticMesh>(nullptr, *UActorFactoryBasicShape::BasicCylinder.ToString())), FName("ClassThumbnail.Cylinder"), FName("ClassIcon.Cylinder"), GetBasicShapeColorOverride(), TOptional<int32>(), NSLOCTEXT("PlacementMode", "Cylinder", "Cylinder"))));
-	Category->Items.Add(CreateID(), MakeShareable(new FPlaceableItem(*UActorFactoryBasicShape::StaticClass(), FAssetData(LoadObject<UStaticMesh>(nullptr, *UActorFactoryBasicShape::BasicCone.ToString())), FName("ClassThumbnail.Cone"), FName("ClassIcon.Cone"), GetBasicShapeColorOverride(), TOptional<int32>(), NSLOCTEXT("PlacementMode", "Cone", "Cone"))));
-	Category->Items.Add(CreateID(), MakeShareable(new FPlaceableItem(*UActorFactoryBasicShape::StaticClass(), FAssetData(LoadObject<UStaticMesh>(nullptr, *UActorFactoryBasicShape::BasicPlane.ToString())), FName("ClassThumbnail.Plane"), FName("ClassIcon.Plane"), GetBasicShapeColorOverride(), TOptional<int32>(), NSLOCTEXT("PlacementMode", "Plane", "Plane"))));
+	Category->Items.Add(CreateID(), MakeShared<FPlaceableItem>(*UActorFactoryEmptyActor::StaticClass()));
+	Category->Items.Add(CreateID(), MakeShared<FPlaceableItem>(*UActorFactoryCharacter::StaticClass()));
+	Category->Items.Add(CreateID(), MakeShared<FPlaceableItem>(*UActorFactoryPawn::StaticClass()));
+	Category->Items.Add(CreateID(), MakeShared<FPlaceableItem>(*UActorFactoryBasicShape::StaticClass(), FAssetData(LoadObject<UStaticMesh>(nullptr, *UActorFactoryBasicShape::BasicCube.ToString())), FName("ClassThumbnail.Cube"), FName("ClassIcon.Cube"), GetBasicShapeColorOverride(), TOptional<int32>(), LOCTEXT("Cube", "Cube")));
+	Category->Items.Add(CreateID(), MakeShared<FPlaceableItem>(*UActorFactoryBasicShape::StaticClass(), FAssetData(LoadObject<UStaticMesh>(nullptr, *UActorFactoryBasicShape::BasicSphere.ToString())), FName("ClassThumbnail.Sphere"), FName("ClassIcon.Sphere"), GetBasicShapeColorOverride(), TOptional<int32>(), LOCTEXT("Sphere", "Sphere")));
+	Category->Items.Add(CreateID(), MakeShared<FPlaceableItem>(*UActorFactoryBasicShape::StaticClass(), FAssetData(LoadObject<UStaticMesh>(nullptr, *UActorFactoryBasicShape::BasicCylinder.ToString())), FName("ClassThumbnail.Cylinder"), FName("ClassIcon.Cylinder"), GetBasicShapeColorOverride(), TOptional<int32>(), LOCTEXT("Cylinder", "Cylinder")));
+	Category->Items.Add(CreateID(), MakeShared<FPlaceableItem>(*UActorFactoryBasicShape::StaticClass(), FAssetData(LoadObject<UStaticMesh>(nullptr, *UActorFactoryBasicShape::BasicCone.ToString())), FName("ClassThumbnail.Cone"), FName("ClassIcon.Cone"), GetBasicShapeColorOverride(), TOptional<int32>(), LOCTEXT("Cone", "Cone")));
+	Category->Items.Add(CreateID(), MakeShared<FPlaceableItem>(*UActorFactoryBasicShape::StaticClass(), FAssetData(LoadObject<UStaticMesh>(nullptr, *UActorFactoryBasicShape::BasicPlane.ToString())), FName("ClassThumbnail.Plane"), FName("ClassIcon.Plane"), GetBasicShapeColorOverride(), TOptional<int32>(), LOCTEXT("Plane", "Plane")));
 
 	// Make a map of UClasses to ActorFactories that support them
 	const TArray< UActorFactory *>& ActorFactories = GEditor->ActorFactories;
@@ -866,7 +868,7 @@ void FPlacementModeModule::RefreshAllPlaceableClasses()
 				continue;
 			}
 
-			Category->Items.Add(CreateID(), MakeShareable(new FPlaceableItem(ActorFactory, FAssetData(*ClassIt))));
+			Category->Items.Add(CreateID(), MakeShared<FPlaceableItem>(ActorFactory, FAssetData(*ClassIt)));
 		}
 	}
 
@@ -941,3 +943,5 @@ void FPlacementModeModule::OnCategoryPermissionListChanged()
 {
 	PlacementModeCategoryListChanged.Broadcast();
 }
+
+#undef LOCTEXT_NAMESPACE
