Placement Module Forcing the loading of assets on Editor boot

The current implementation of `FPlacementModeModule::RefreshAllPlaceableClasses()` forces all Actor Classes from Actor Factories to be loaded, even though those classes might not be placeable.
Some of these classes can be huge in bigger games, causing significant delays on Editor Loading.

The purpose of this change is to achieve a similar outcome using `FTopLevelAssetPath` instead of loading the `UClass` if it is not already loaded.
This is achieved by leveraging the information present in the AssetRegistry, therefore further delaying the initial call to `FPlacementModeModule::RefreshAllPlaceableClasses()` until after the AssetRegistry is done loading.

diff --git a/Engine/Source/Editor/LevelEditor/Private/LevelEditor.cpp b/Engine/Source/Editor/LevelEditor/Private/LevelEditor.cpp
--- a/Engine/Source/Editor/LevelEditor/Private/LevelEditor.cpp
+++ b/Engine/Source/Editor/LevelEditor/Private/LevelEditor.cpp
@@ -308,6 +308,7 @@ void FLevelEditorModule::StartupModule()
 	MessageLogModule.RegisterLogListing("BuildAndSubmitErrors", LOCTEXT("BuildAndSubmitErrors", "Build and Submit Errors"));
 
 	OutlinerSettings = MakeShared<FLevelEditorOutlinerSettings>();
+	OutlinerSettings->Initialize();
 	OutlinerSettings->SetupBuiltInCategories();
 }
 
diff --git a/Engine/Source/Editor/LevelEditor/Private/LevelEditorOutlinerSettings.cpp b/Engine/Source/Editor/LevelEditor/Private/LevelEditorOutlinerSettings.cpp
--- a/Engine/Source/Editor/LevelEditor/Private/LevelEditorOutlinerSettings.cpp
+++ b/Engine/Source/Editor/LevelEditor/Private/LevelEditorOutlinerSettings.cpp
@@ -36,6 +36,27 @@
 const FString FLevelEditorOutlinerSettings::UnsavedAssetsFilterName("UnsavedAssetsFilter");
 const FString FLevelEditorOutlinerSettings::UncontrolledAssetsFilterName("UncontrolledAssetsFilter");
 
+FLevelEditorOutlinerSettings::~FLevelEditorOutlinerSettings()
+{
+	if (IPlacementModeModule::IsAvailable())
+	{
+		IPlacementModeModule::Get().OnAllPlaceableAssetsChanged().RemoveAll(this);
+	}
+}
+
+void FLevelEditorOutlinerSettings::Initialize()
+{
+	FDelayedAutoRegisterHelper EndOfEngineInit(EDelayedRegisterRunPhase::EndOfEngineInit, [WeakThis = AsWeak()]()
+	{
+		if (const TSharedPtr<FLevelEditorOutlinerSettings> This = WeakThis.Pin())
+		{
+			IPlacementModeModule& PlacementModeModule = IPlacementModeModule::Get();
+			// ensure all the CustomClassFilters are up to date by rerunning CreateDefaultFilters when new assets are loaded
+			PlacementModeModule.OnAllPlaceableAssetsChanged().AddSP(This.ToSharedRef(), &FLevelEditorOutlinerSettings::CreateDefaultFilters);
+		}
+	});
+}
+
 void FLevelEditorOutlinerSettings::AddCustomFilter(TSharedRef<FFilterBase<const ISceneOutlinerTreeItem&>> InCustomFilter)
 {
 	CustomFilters.Add(InCustomFilter);
diff --git a/Engine/Source/Editor/LevelEditor/Public/LevelEditorOutlinerSettings.h b/Engine/Source/Editor/LevelEditor/Public/LevelEditorOutlinerSettings.h
--- a/Engine/Source/Editor/LevelEditor/Public/LevelEditorOutlinerSettings.h
+++ b/Engine/Source/Editor/LevelEditor/Public/LevelEditorOutlinerSettings.h
@@ -38,7 +38,11 @@ struct FLevelEditorOutlinerBuiltInCategories
 class LEVELEDITOR_API FLevelEditorOutlinerSettings :  public TSharedFromThis<class FLevelEditorOutlinerSettings>
 {
 public:
-
+	FLevelEditorOutlinerSettings(){}
+	~FLevelEditorOutlinerSettings();
+	
+	void Initialize();
+	
 	/** Delegate to create a Filter for the Outliner */
 	DECLARE_DELEGATE_RetVal(TSharedPtr<FFilterBase<const ISceneOutlinerTreeItem&>>, FOutlinerFilterFactory);
 
diff --git a/Engine/Source/Editor/PlacementMode/PlacementMode.Build.cs b/Engine/Source/Editor/PlacementMode/PlacementMode.Build.cs
--- a/Engine/Source/Editor/PlacementMode/PlacementMode.Build.cs
+++ b/Engine/Source/Editor/PlacementMode/PlacementMode.Build.cs
@@ -8,11 +8,16 @@ public class PlacementMode : ModuleRules
     {
 		PrivateIncludePathModuleNames.Add("AssetTools");
 
+        PublicDependencyModuleNames.AddRange( 
+            new string[] { 
+                "Engine", 
+            } 
+        );
+        
         PrivateDependencyModuleNames.AddRange( 
             new string[] { 
                 "Core", 
                 "CoreUObject",
-                "Engine", 
                 "InputCore",
                 "Slate",
 				"SlateCore",
diff --git a/Engine/Source/Editor/PlacementMode/Private/PlacementModeLog.cpp b/Engine/Source/Editor/PlacementMode/Private/PlacementModeLog.cpp
new file mode 100644
--- /dev/null
+++ b/Engine/Source/Editor/PlacementMode/Private/PlacementModeLog.cpp
@@ -0,0 +1,4 @@
+﻿// Copyright Epic Games, Inc. All Rights Reserved.
+#include "PlacementModeLog.h"
+
+DEFINE_LOG_CATEGORY(LogPlacementMode);
\ No newline at end of file
diff --git a/Engine/Source/Editor/PlacementMode/Private/PlacementModeLog.h b/Engine/Source/Editor/PlacementMode/Private/PlacementModeLog.h
new file mode 100644
--- /dev/null
+++ b/Engine/Source/Editor/PlacementMode/Private/PlacementModeLog.h
@@ -0,0 +1,7 @@
+﻿// Copyright Epic Games, Inc. All Rights Reserved.
+
+#pragma once
+
+#include "Logging/LogMacros.h"
+
+DECLARE_LOG_CATEGORY_EXTERN(LogPlacementMode, Log, All);
\ No newline at end of file
diff --git a/Engine/Source/Editor/PlacementMode/Private/PlacementModeModule.cpp b/Engine/Source/Editor/PlacementMode/Private/PlacementModeModule.cpp
--- a/Engine/Source/Editor/PlacementMode/Private/PlacementModeModule.cpp
+++ b/Engine/Source/Editor/PlacementMode/Private/PlacementModeModule.cpp
@@ -55,6 +55,7 @@
 #include "ActorFactories/ActorFactoryPlanarReflection.h"
 #include "SPlacementModeTools.h"
 #include "AssetSelection.h"
+#include "PlacementModeLog.h"
 
 #define LOCTEXT_NAMESPACE "PlacementMode"
 
@@ -508,7 +509,9 @@ void FPlacementModeModule::OnInitialAssetsScanComplete()
 	IAssetRegistry& AssetRegistry = FModuleManager::LoadModuleChecked<FAssetRegistryModule>(TEXT("AssetRegistry")).Get();
 	AssetRegistry.OnAssetAdded().AddRaw(this, &FPlacementModeModule::OnAssetAdded);
 	AssetRegistry.OnFilesLoaded().RemoveAll(this);
-
+	
+	RegenerateItemsForCategory(FBuiltInPlacementCategories::AllClasses());
+	UE_LOG(LogPlacementMode, Display, TEXT("The Asset Registry is done with its initial scan, the list of placeable classes has been updated."))
 	AllPlaceableAssetsChanged.Broadcast();
 }
 
@@ -806,6 +809,12 @@ void FPlacementModeModule::RefreshAllPlaceableClasses()
 		return;
 	}
 
+	TRACE_CPUPROFILER_EVENT_SCOPE_STR("FPlacementModeModule::RefreshAllPlaceableClasses");
+	if (FModuleManager::Get().IsModuleLoaded("AssetRegistry") && IAssetRegistry::Get() && IAssetRegistry::Get()->IsLoadingAssets())
+	{
+		UE_LOG(LogPlacementMode, Display, TEXT("The Asset Registry is not yet fully loaded so some placeable classes might be missing."))
+	}
+	
 	Category->Items.Reset();
 
 
@@ -827,27 +836,101 @@ void FPlacementModeModule::RefreshAllPlaceableClasses()
 	Category->Items.Add(CreateID(), MakeShared<FPlaceableItem>(*UActorFactoryBasicShape::StaticClass(), FAssetData(LoadObject<UStaticMesh>(nullptr, *UActorFactoryBasicShape::BasicCone.ToString())), FName("ClassThumbnail.Cone"), FName("ClassIcon.Cone"), GetBasicShapeColorOverride(), TOptional<int32>(), LOCTEXT("Cone", "Cone")));
 	Category->Items.Add(CreateID(), MakeShared<FPlaceableItem>(*UActorFactoryBasicShape::StaticClass(), FAssetData(LoadObject<UStaticMesh>(nullptr, *UActorFactoryBasicShape::BasicPlane.ToString())), FName("ClassThumbnail.Plane"), FName("ClassIcon.Plane"), GetBasicShapeColorOverride(), TOptional<int32>(), LOCTEXT("Plane", "Plane")));
 
+	// Below, we will be checking if the actor factory classes inherit from certain classes, but these actor factory classes might not be loaded
+	// To handle this, we retrieve the derived classes via the Asset Registry
+	const FAssetRegistryModule& AssetRegistryModule = FModuleManager::LoadModuleChecked<FAssetRegistryModule>("AssetRegistry");
+	TSet<FTopLevelAssetPath> DerivedActorClassNames, DerivedVolumeClassNames, DerivedBrushClassNames;
+	AssetRegistryModule.Get().GetDerivedClassNames({ AActor::StaticClass()->GetClassPathName() }, {}, DerivedActorClassNames);
+	AssetRegistryModule.Get().GetDerivedClassNames({ AVolume::StaticClass()->GetClassPathName() }, {}, DerivedVolumeClassNames);
+	AssetRegistryModule.Get().GetDerivedClassNames({ ABrush::StaticClass()->GetClassPathName() }, {}, DerivedBrushClassNames);
+	
 	// Make a map of UClasses to ActorFactories that support them
 	const TArray< UActorFactory *>& ActorFactories = GEditor->ActorFactories;
-	TMap<UClass*, UActorFactory*> ActorFactoryMap;
+	TMap<FTopLevelAssetPath, UActorFactory*> ActorFactoryPathMap;
+	FAssetData NoAssetData;
+	FText UnusedErrorMessage;
 	for (int32 FactoryIdx = 0; FactoryIdx < ActorFactories.Num(); ++FactoryIdx)
 	{
 		UActorFactory* ActorFactory = ActorFactories[FactoryIdx];
+		if (!ActorFactory)
+		{
+			continue;
+		}
+		
+		FTopLevelAssetPath ActorClassPath{};
+		if (ActorFactory->NewActorClass)
+		{
+			ActorClassPath = ActorFactory->NewActorClass->GetClassPathName();
+		}
+		else if (!ActorFactory->NewActorClassName.IsEmpty())
+		{
+			if (FPackageName::IsShortPackageName(ActorFactory->NewActorClassName))
+			{
+				// We cannot create a FTopLevelAssetPath from a short package name so we try resolving the name, which is what would happen if we were using
+				// UActorFactory::GetDefaultActorClass. The parameters for ResolveName are the same to the ones used in StaticLoadObject, but without raising errors
+				UObject* OuterPackage = nullptr;
+				FString ResolvedName = ActorFactory->NewActorClassName;
+				if (ResolveName(OuterPackage, ResolvedName, true, false, LOAD_NoWarn))
+				{
+					ActorClassPath = FTopLevelAssetPath{ OuterPackage->GetFName(), FName{ResolvedName} };
+				}
+			}
+			else
+			{
+				ActorClassPath = FTopLevelAssetPath{ ActorFactory->NewActorClassName };
+			}
+		}
 
-		if (ActorFactory)
+		// This is the equivalent check as the one in the loop below but using the cached data from the asset registry:
+		// We only want an Actor that is not a Brush unless it is a Volume
+		if (!DerivedActorClassNames.Contains(ActorClassPath) ||
+			(DerivedBrushClassNames.Contains(ActorClassPath) && !DerivedVolumeClassNames.Contains(ActorClassPath)))
 		{
-			ActorFactoryMap.Add(ActorFactory->GetDefaultActorClass(FAssetData()), ActorFactory);
+			continue;
 		}
-	}
+		
+		const bool bSkipARFilteredAssets = false; // Allow the retrieval of BlueprintGeneratedClasses
+		FAssetData ActorClassAssetData = AssetRegistryModule.Get().GetAssetByObjectPath(FSoftObjectPath{ActorClassPath}, false, bSkipARFilteredAssets);
 
-	FAssetData NoAssetData;
-	FText UnusedErrorMessage;
+		// Only Blueprint classes might not be loaded at this point
+		int32 BPClassFlagsInt;
+		if (!ActorClassAssetData.GetTagValue(FBlueprintTags::ClassFlags, BPClassFlagsInt) || ActorClassAssetData.IsAssetLoaded())
+		{
+			// If the class is already loaded, it will be picked up in the class iterator loop below
+			ActorFactoryPathMap.Add(ActorClassPath, ActorFactory);
+			continue; 
+		}
+			
+		const EClassFlags BPClassFlags = static_cast<EClassFlags>(BPClassFlagsInt);
+		if (EnumHasAllFlags(BPClassFlags, CLASS_NotPlaceable) ||
+			EnumHasAnyFlags(BPClassFlags, CLASS_Abstract | CLASS_Deprecated | CLASS_NewerVersionExists))
+		{
+			continue;
+		}
+
+		const bool bIsVolume = DerivedVolumeClassNames.Contains(ActorClassPath);
+		if (bIsVolume)
+		{
+			// Retrieving the ActorFactory for a volume requires a call to GEditor->FindActorFactoryByClassForActorClass which requires the class
+			// to be loaded to access its default actor. We therefore need to load the class here, but now that the class is loaded, we will defer
+			// the processing of this Volume to the loop of loaded classes below, otherwise we would double up the entry for this AssetClass.
+			const UClass* AssetClass = Cast<UClass>(ActorClassAssetData.GetAsset());
+			ActorFactoryPathMap.Add(ActorClassPath, ActorFactory);
+			continue; 
+		}
+		else if (ActorFactory && !ActorFactory->CanCreateActorFrom(NoAssetData, UnusedErrorMessage))
+		{
+			continue;
+		}
+			
+		Category->Items.Add(CreateID(), MakeShared<FPlaceableItem>(ActorFactory, ActorClassAssetData));
+	}
 
 	// Add loaded classes
 	for (TObjectIterator<UClass> ClassIt; ClassIt; ++ClassIt)
 	{
 		// Don't offer skeleton classes
-		bool bIsSkeletonClass = FKismetEditorUtilities::IsClassABlueprintSkeleton(*ClassIt);
+		const bool bIsSkeletonClass = FKismetEditorUtilities::IsClassABlueprintSkeleton(*ClassIt);
 
 		if (!ClassIt->HasAllClassFlags(CLASS_NotPlaceable) &&
 			!ClassIt->HasAnyClassFlags(CLASS_Abstract | CLASS_Deprecated | CLASS_NewerVersionExists) &&
@@ -855,11 +938,10 @@ void FPlacementModeModule::RefreshAllPlaceableClasses()
 			(!ClassIt->IsChildOf(ABrush::StaticClass()) || ClassIt->IsChildOf(AVolume::StaticClass())) &&
 			!bIsSkeletonClass)
 		{
-			UActorFactory* ActorFactory = ActorFactoryMap.FindRef(*ClassIt);
-
-			const bool IsVolume = ClassIt->IsChildOf(AVolume::StaticClass());
+			UActorFactory* ActorFactory = ActorFactoryPathMap.FindRef(ClassIt->GetClassPathName());
 
-			if (IsVolume)
+			const bool bIsVolume = ClassIt->IsChildOf(AVolume::StaticClass());
+			if (bIsVolume)
 			{
 				ActorFactory = GEditor->FindActorFactoryByClassForActorClass(UActorFactoryBoxVolume::StaticClass(), *ClassIt);
 			}
diff --git a/Engine/Source/Editor/PlacementMode/Public/IPlacementModeModule.h b/Engine/Source/Editor/PlacementMode/Public/IPlacementModeModule.h
--- a/Engine/Source/Editor/PlacementMode/Public/IPlacementModeModule.h
+++ b/Engine/Source/Editor/PlacementMode/Public/IPlacementModeModule.h
@@ -13,6 +13,7 @@
 #include "ActorPlacementInfo.h"
 #include "GameFramework/Volume.h"
 #include "Editor.h"
+#include "Engine/BlueprintGeneratedClass.h"
 #include "Textures/SlateIcon.h"
 
 class FNamePermissionList;
@@ -232,7 +233,12 @@ struct FPlaceableItem
 		else
 		{
 			NativeName = AssetData.AssetName.ToString();
-			DisplayName = FText::FromName(AssetData.AssetName);
+			if (AssetData.AssetClassPath == UBlueprintGeneratedClass::StaticClass()->GetClassPathName())
+			{
+				// For BPGC, we do not want to display the "_C" suffix, as the name of the loaded class will not have it.
+				NativeName.RemoveFromEnd(TEXT("_C"));
+			}
+			DisplayName = FText::FromString(NativeName);
 		}
 	}
 
