MotionDesign : Editor
- Fix editor AddQuickMenu to allow for multiple FPlaceableItem even if they share same FAssetData, this caused variants of an actor that uses different components to only be allowed once in that menu. Eg: AAvaShapeActor

diff --git a/Engine/Source/Editor/PlacementMode/Private/PlacementModeModule.cpp b/Engine/Source/Editor/PlacementMode/Private/PlacementModeModule.cpp
--- a/Engine/Source/Editor/PlacementMode/Private/PlacementModeModule.cpp
+++ b/Engine/Source/Editor/PlacementMode/Private/PlacementModeModule.cpp
@@ -299,8 +299,10 @@ void FPlacementModeModule::StartupModule()
 			for (auto& Pair : Category->Items)
 			{
 				TSharedPtr<const FPlaceableItem> Item = Pair.Value;
+
+				// Use NativeName instead of AssetName to allow for multiple item to be registered otherwise can't have different item sharing same FAssetData
 				FToolMenuEntry& NewEntry = Section.AddEntry(
-					FToolMenuEntry::InitWidget(Item->AssetData.AssetName, SNew(SPlacementAssetMenuEntry, Item), FText(), true, true)
+					FToolMenuEntry::InitWidget(FName(Item->NativeName), SNew(SPlacementAssetMenuEntry, Item), FText(), true, true)
 				);
 
 				if (++count > MaxItems)
@@ -332,8 +334,9 @@ void FPlacementModeModule::StartupModule()
 						for (auto& Pair : Category->Items)
 						{
 							TSharedPtr<const FPlaceableItem> Item = Pair.Value;
+							// Use NativeName instead of AssetName to allow for multiple item to be registered otherwise can't have different item sharing same FAssetData
 							FToolMenuEntry& NewEntry = Section.AddEntry(
-								FToolMenuEntry::InitWidget(Item->AssetData.AssetName, SNew(SPlacementAssetMenuEntry, Item), FText(), true, true)
+								FToolMenuEntry::InitWidget(FName(Item->NativeName), SNew(SPlacementAssetMenuEntry, Item), FText(), true, true)
 							);
 						}
 					})
