[BUGFIX] Entity Placement Factory is creating actors in all the cases.

#rnx

#Synced-CL 50223086

diff --git a/Engine/Source/Editor/PlacementMode/PlacementMode.Build.cs b/Engine/Source/Editor/PlacementMode/PlacementMode.Build.cs
--- a/Engine/Source/Editor/PlacementMode/PlacementMode.Build.cs
+++ b/Engine/Source/Editor/PlacementMode/PlacementMode.Build.cs
@@ -31,7 +31,9 @@ public class PlacementMode : ModuleRules
                 "AssetTools",
                 "EditorWidgets",
                 "ToolMenus",
-                "WidgetRegistration"
+				"TypedElementFramework",
+				"TypedElementRuntime",
+				"WidgetRegistration"
             } 
         );
     }
diff --git a/Engine/Source/Editor/PlacementMode/Private/SPlacementModeTools.cpp b/Engine/Source/Editor/PlacementMode/Private/SPlacementModeTools.cpp
--- a/Engine/Source/Editor/PlacementMode/Private/SPlacementModeTools.cpp
+++ b/Engine/Source/Editor/PlacementMode/Private/SPlacementModeTools.cpp
@@ -10,11 +10,15 @@
 #include "Styling/AppStyle.h"
 #include "EditorModeManager.h"
 #include "EditorModes.h"
+#include "Elements/Actor/ActorElementData.h"
+#include "Elements/Framework/TypedElementRegistry.h"
+#include "Elements/Interfaces/TypedElementObjectInterface.h"
 #include "Framework/Application/SlateApplication.h"
 #include "AssetThumbnail.h"
 #include "LevelEditor.h"
 #include "LevelEditorActions.h"
 #include "LevelEditorViewport.h"
+#include "LevelEditorSubsystem.h"
 #include "ContentBrowserDataDragDropOp.h"
 #include "EditorClassUtils.h"
 #include "Widgets/Input/SSearchBox.h"
@@ -39,6 +43,13 @@ namespace PlacementModeTools
 {
 	bool bItemInternalsInTooltip = false;
 	FAutoConsoleVariableRef CVarItemInternalsInTooltip(TEXT("PlacementMode.ItemInternalsInTooltip"), bItemInternalsInTooltip, TEXT("Shows placeable item internal information in its tooltip"));
+
+	bool bAllowAssetFactoryInterfaceWorkflowForItems = true;
+	FAutoConsoleVariableRef CVarAllowAssetFactoryInterfaceWorkflowForItems(
+		TEXT("PlacementMode.AllowAssetFactoryInterfaceWorkflowForItems"), 
+		bAllowAssetFactoryInterfaceWorkflowForItems,
+		TEXT("If enabled AssetFactoryInterface would be considered when placing actors and not always the default ActorFactory would be used.")
+	);
 }
 
 struct FSortPlaceableItems
@@ -529,6 +540,45 @@ const FSlateBrush* SPlacementAssetMenuEntry::GetIcon() const
 	return AssetImage;
 }
 
+bool SPlacementAssetMenuEntry::TryPlacingItemInLevel(TNotNull<ULevel*> DesiredLevel, AActor*& OutPlacedActor)
+{
+	bool bPlacementHandled = false;
+	if (TScriptInterface<IAssetFactoryInterface> AssetFactory = Item->AssetFactory)
+	{
+		if (UObject* LoadedAsset = Item->AssetData.GetAsset())
+		{
+			UE::AssetPlacementUtil::FExtraPlaceAssetOptions PlacementOptions;
+			PlacementOptions.bSelectOutput = true;
+			PlacementOptions.ObjectFlags = RF_Transactional;
+			PlacementOptions.FactoryToUse = AssetFactory;
+			PlacementOptions.Name = NAME_None;
+
+			TArray<FTypedElementHandle> PlacedItemsHandles = FLevelEditorViewportClient::TryPlacingAssetObject(DesiredLevel, LoadedAsset, PlacementOptions);
+			bPlacementHandled = !PlacedItemsHandles.IsEmpty();
+
+			// Placement succeeded. Not necessarily an actor, but could be elements.
+			if (bPlacementHandled)
+			{
+				for (const FTypedElementHandle& Element : PlacedItemsHandles)
+				{
+					if (AActor* Actor = ActorElementDataUtil::GetActorFromHandle(Element))
+					{
+						OutPlacedActor = Actor; 
+						break;
+					}
+				}
+
+				if (IPlacementModeModule::IsAvailable())
+				{
+					IPlacementModeModule::Get().AddToRecentlyPlaced(LoadedAsset, AssetFactory);
+				}
+			}
+		}
+	}
+
+	return bPlacementHandled;
+}
+
 
 FReply SPlacementAssetMenuEntry::OnMouseButtonDown(const FGeometry& MyGeometry, const FPointerEvent& MouseEvent)
 {
@@ -548,23 +598,29 @@ FReply SPlacementAssetMenuEntry::OnMouseButtonUp(const FGeometry& MyGeometry, co
 	{
 		bIsPressed = false;
 
-		UActorFactory* Factory = Item->Factory;
-		if (!Item->Factory)
-		{
-			// If no actor factory was found or failed, add the actor from the uclass
-			UClass* AssetClass = Item->AssetData.GetClass();
-			if (AssetClass)
-			{
-				UObject* ClassObject = AssetClass->GetDefaultObject();
-				FActorFactoryAssetProxy::GetFactoryForAssetObject(ClassObject);
-			}
-		}
-
 		{
 			// Note: Capture the add and the move within a single transaction, so that the placed actor position is calculated correctly by the transaction diff
 			FScopedTransaction Transaction(NSLOCTEXT("UnrealEd", "CreateActor", "Create Actor"));
+			
+			AActor* NewActor = nullptr;
+			bool bPlacementHandled = false;
+			if (PlacementModeTools::bAllowAssetFactoryInterfaceWorkflowForItems)
+			{
+				if (ULevelEditorSubsystem* LevelEditorSubsystem = GEditor->GetEditorSubsystem<ULevelEditorSubsystem>())
+				{
+					if (ULevel* DesiredLevel = LevelEditorSubsystem->GetCurrentLevel())
+					{
+						bPlacementHandled = TryPlacingItemInLevel(DesiredLevel, NewActor);
+					}
+				}
+			}
+			
+			// Likely not an element placement, try adding actor instead.
+			if (!bPlacementHandled)
+			{
+				NewActor = FLevelEditorActionCallbacks::AddActor(Item->Factory, Item->AssetData, nullptr);
+			}
 
-			AActor* NewActor = FLevelEditorActionCallbacks::AddActor(Factory, Item->AssetData, nullptr);
 			if (NewActor && GCurrentLevelEditingViewportClient)
 			{
 				GEditor->MoveActorInFrontOfCamera(*NewActor,
diff --git a/Engine/Source/Editor/PlacementMode/Private/SPlacementModeTools.h b/Engine/Source/Editor/PlacementMode/Private/SPlacementModeTools.h
--- a/Engine/Source/Editor/PlacementMode/Private/SPlacementModeTools.h
+++ b/Engine/Source/Editor/PlacementMode/Private/SPlacementModeTools.h
@@ -14,6 +14,7 @@
 #include "Widgets/Views/STableRow.h"
 #include "Misc/TextFilter.h"
 
+class ULevel;
 class STextBlock;
 class FCategoryDrivenContentBuilder;
 
@@ -83,6 +84,8 @@ public:
 	TSharedPtr<const FPlaceableItem> Item;
 
 	virtual FSlateColor GetForegroundColor() const override;
+private:
+	bool TryPlacingItemInLevel(TNotNull<ULevel*> DesiredLevel, AActor*& OutPlacedActor);
 
 private:
 	const FSlateBrush* GetBorder() const;
