Disable Placement hold for commandlets

#Synced-CL 51322414

#Synced-CL 51322414

diff --git a/Engine/Source/Editor/PlacementMode/Private/PlacementModeModule.cpp b/Engine/Source/Editor/PlacementMode/Private/PlacementModeModule.cpp
--- a/Engine/Source/Editor/PlacementMode/Private/PlacementModeModule.cpp
+++ b/Engine/Source/Editor/PlacementMode/Private/PlacementModeModule.cpp
@@ -107,16 +107,19 @@ void FPlacementModeModule::StartupModule()
 		RecentlyPlaced.Add(FActorPlacementInfo(RecentlyPlacedAsStrings[Index]));
 	}
 
-	IAssetRegistry& AssetRegistry = FModuleManager::LoadModuleChecked<FAssetRegistryModule>(TEXT("AssetRegistry")).Get();
-	AssetRegistry.OnAssetRemoved().AddRaw(this, &FPlacementModeModule::OnAssetRemoved);
-	AssetRegistry.OnAssetRenamed().AddRaw(this, &FPlacementModeModule::OnAssetRenamed);
-	if (AssetRegistry.IsLoadingAssets())
+	if (!IsRunningCommandlet())
 	{
-		AssetRegistry.OnFilesLoaded().AddRaw(this, &FPlacementModeModule::OnInitialAssetsScanComplete);
-	}
-	else
-	{
-		AssetRegistry.OnAssetAdded().AddRaw(this, &FPlacementModeModule::OnAssetAdded);
+		IAssetRegistry& AssetRegistry = FModuleManager::LoadModuleChecked<FAssetRegistryModule>(TEXT("AssetRegistry")).Get();
+		AssetRegistry.OnAssetRemoved().AddRaw(this, &FPlacementModeModule::OnAssetRemoved);
+		AssetRegistry.OnAssetRenamed().AddRaw(this, &FPlacementModeModule::OnAssetRenamed);
+		if (AssetRegistry.IsLoadingAssets())
+		{
+			AssetRegistry.OnFilesLoaded().AddRaw(this, &FPlacementModeModule::OnInitialAssetsScanComplete);
+		}
+		else
+		{
+			AssetRegistry.OnAssetAdded().AddRaw(this, &FPlacementModeModule::OnAssetAdded);
+		}
 	}
 
 	TOptional<FLinearColor> BasicShapeColorOverride = GetBasicShapeColorOverride();
