[Backout] - CL51426446
#rnx
Original CL Desc
-----------------------------------------------------------------
Disable Placement hold for commandlets

diff --git a/Engine/Source/Editor/PlacementMode/Private/PlacementModeModule.cpp b/Engine/Source/Editor/PlacementMode/Private/PlacementModeModule.cpp
--- a/Engine/Source/Editor/PlacementMode/Private/PlacementModeModule.cpp
+++ b/Engine/Source/Editor/PlacementMode/Private/PlacementModeModule.cpp
@@ -107,19 +107,16 @@ void FPlacementModeModule::StartupModule()
 		RecentlyPlaced.Add(FActorPlacementInfo(RecentlyPlacedAsStrings[Index]));
 	}
 
-	if (!IsRunningCommandlet())
+	IAssetRegistry& AssetRegistry = FModuleManager::LoadModuleChecked<FAssetRegistryModule>(TEXT("AssetRegistry")).Get();
+	AssetRegistry.OnAssetRemoved().AddRaw(this, &FPlacementModeModule::OnAssetRemoved);
+	AssetRegistry.OnAssetRenamed().AddRaw(this, &FPlacementModeModule::OnAssetRenamed);
+	if (AssetRegistry.IsLoadingAssets())
 	{
-		IAssetRegistry& AssetRegistry = FModuleManager::LoadModuleChecked<FAssetRegistryModule>(TEXT("AssetRegistry")).Get();
-		AssetRegistry.OnAssetRemoved().AddRaw(this, &FPlacementModeModule::OnAssetRemoved);
-		AssetRegistry.OnAssetRenamed().AddRaw(this, &FPlacementModeModule::OnAssetRenamed);
-		if (AssetRegistry.IsLoadingAssets())
-		{
-			AssetRegistry.OnFilesLoaded().AddRaw(this, &FPlacementModeModule::OnInitialAssetsScanComplete);
-		}
-		else
-		{
-			AssetRegistry.OnAssetAdded().AddRaw(this, &FPlacementModeModule::OnAssetAdded);
-		}
+		AssetRegistry.OnFilesLoaded().AddRaw(this, &FPlacementModeModule::OnInitialAssetsScanComplete);
+	}
+	else
+	{
+		AssetRegistry.OnAssetAdded().AddRaw(this, &FPlacementModeModule::OnAssetAdded);
 	}
 
 	TOptional<FLinearColor> BasicShapeColorOverride = GetBasicShapeColorOverride();
